SlateIM: Experimental immediate mode wrapper for Slate

Intended for creating tools. Available in the editor, in cooked builds, and for programs without the engine.

See SlateIM.h for the available API

#rnx

diff --git a/Engine/Plugins/Experimental/SlateIM/SlateIM.uplugin b/Engine/Plugins/Experimental/SlateIM/SlateIM.uplugin
new file mode 100644
--- /dev/null
+++ b/Engine/Plugins/Experimental/SlateIM/SlateIM.uplugin
@@ -0,0 +1,25 @@
+{
+	"FileVersion": 3,
+	"Version": 1,
+	"VersionName": "1.0",
+	"FriendlyName": "SlateIM",
+	"Description": "An immediate mode wrapper for Slate. Intended for building debugging tools.",
+	"Category": "UI",
+	"CreatedBy": "Epic Games, Inc.",
+	"CreatedByURL": "https://epicgames.com",
+	"DocsURL": "",
+	"MarketplaceURL": "",
+	"SupportURL": "",
+	"CanContainContent": false,
+	"IsBetaVersion": false,
+	"IsExperimentalVersion": true,
+	"Installed": false,
+	"EnabledByDefault": false,
+	"Modules": [
+		{
+			"Name": "SlateIM",
+			"Type": "Runtime",
+			"LoadingPhase": "Default"
+		}
+	]
+}
\ No newline at end of file
diff --git a/Engine/Plugins/Experimental/SlateIM/Source/SlateIM/Private/Containers/ISlateIMContainer.h b/Engine/Plugins/Experimental/SlateIM/Source/SlateIM/Private/Containers/ISlateIMContainer.h
new file mode 100644
--- /dev/null
+++ b/Engine/Plugins/Experimental/SlateIM/Source/SlateIM/Private/Containers/ISlateIMContainer.h
@@ -0,0 +1,27 @@
+// Copyright Epic Games, Inc. All Rights Reserved.
+
+#pragma once
+
+#include "Misc/ISlateIMChild.h"
+#include "Misc/SlateIMTypeChecking.h"
+#include "Templates/SharedPointer.h"
+#include "Widgets/SWidget.h"
+
+struct FSlateIMSlotData;
+
+class FChildren;
+class SWidget;
+
+class ISlateIMContainer : public ISlateIMTypeChecking
+{
+public:
+	virtual FString GetDebugName() { return GetContainer().GetWidget() ? *GetContainer().GetWidget()->GetTypeAsString() : TEXT("UNKNOWN"); }
+
+	virtual int32 GetNumChildren() = 0;
+	virtual FSlateIMChild GetChild(int32 Index) = 0;
+	virtual void UpdateChild(FSlateIMChild Child, int32 Index, const FSlateIMSlotData& AlignmentData) = 0;
+
+	virtual FSlateIMChild GetContainer() = 0;
+
+	virtual void RemoveUnusedChildren(int32 LastUsedChildIndex) = 0;
+};
diff --git a/Engine/Plugins/Experimental/SlateIM/Source/SlateIM/Private/Containers/SImCompoundWidget.cpp b/Engine/Plugins/Experimental/SlateIM/Source/SlateIM/Private/Containers/SImCompoundWidget.cpp
new file mode 100644
--- /dev/null
+++ b/Engine/Plugins/Experimental/SlateIM/Source/SlateIM/Private/Containers/SImCompoundWidget.cpp
@@ -0,0 +1,72 @@
+﻿// Copyright Epic Games, Inc. All Rights Reserved.
+
+#include "SImCompoundWidget.h"
+
+#include "Misc/SlateIMSlotData.h"
+#include "SImStackBox.h"
+
+
+SLATE_IMPLEMENT_WIDGET(SImCompoundWidget)
+
+void SImCompoundWidget::PrivateRegisterAttributes(FSlateAttributeInitializer& AttributeInitializer)
+{
+}
+
+int32 SImCompoundWidget::GetNumChildren()
+{
+	if (Container)
+	{
+		return Container->GetNumChildren();
+	}
+	
+	return GetChildren()->Num();
+}
+
+FSlateIMChild SImCompoundWidget::GetChild(int32 Index)
+{
+	if (Container)
+	{
+		return Container->GetChild(Index);
+	}
+	
+	if (Index >= 0 && Index < GetChildren()->Num())
+	{
+		return GetChildren()->GetChildAt(Index);
+	}
+
+	return nullptr;
+}
+
+void SImCompoundWidget::RemoveUnusedChildren(int32 LastUsedChildIndex)
+{
+	if (Container)
+	{
+		Container->RemoveUnusedChildren(LastUsedChildIndex);
+	}
+}
+
+void SImCompoundWidget::UpdateChild(FSlateIMChild Child, int32 Index, const FSlateIMSlotData& AlignmentData)
+{
+	if (!Container)
+	{
+		Container = SNew(SImStackBox)
+			.Orientation(Orient_Vertical);
+		ChildSlot
+		.Padding(FMargin(0))
+		.HAlign(HAlign_Fill)
+		.VAlign(VAlign_Fill)
+		[
+			Container.ToSharedRef()
+		];
+	}
+	
+	Container->UpdateChild(Child, Index, AlignmentData);
+}
+
+void SImCompoundWidget::SetOrientation(EOrientation InOrientation)
+{
+	if (Container)
+	{
+		Container->SetOrientation(InOrientation);
+	}
+}
diff --git a/Engine/Plugins/Experimental/SlateIM/Source/SlateIM/Private/Containers/SImCompoundWidget.h b/Engine/Plugins/Experimental/SlateIM/Source/SlateIM/Private/Containers/SImCompoundWidget.h
new file mode 100644
--- /dev/null
+++ b/Engine/Plugins/Experimental/SlateIM/Source/SlateIM/Private/Containers/SImCompoundWidget.h
@@ -0,0 +1,73 @@
+﻿// Copyright Epic Games, Inc. All Rights Reserved.
+
+#pragma once
+
+#include "ISlateIMContainer.h"
+#include "Widgets/Layout/SBorder.h"
+
+class SImStackBox;
+
+class SImCompoundWidget : public SBorder, public ISlateIMContainer
+{
+	SLATE_DECLARE_WIDGET(SImCompoundWidget, SBorder)
+	SLATE_IM_TYPE_DATA(SImCompoundWidget, ISlateIMContainer)
+	
+public:
+	virtual int32 GetNumChildren() override;
+	virtual FSlateIMChild GetChild(int32 Index) override;
+
+	virtual FSlateIMChild GetContainer() override
+	{
+		return AsShared();
+	}
+
+	virtual void RemoveUnusedChildren(int32 LastUsedChildIndex) override;
+
+	virtual void UpdateChild(FSlateIMChild Child, int32 Index, const FSlateIMSlotData& AlignmentData) override;
+
+	void SetBackgroundImage(const FSlateBrush* BorderBrush)
+	{
+		SetBorderImage(BorderBrush);
+	}
+
+	void SetContentPadding(const FMargin ContentPadding)
+	{
+		SetPadding(ContentPadding);
+	}
+
+	void SetAbsorbMouse(bool bInAbsorb)
+	{
+		bAbsorbMouse = bInAbsorb;
+	}
+
+	void SetOrientation(EOrientation InOrientation);
+
+	virtual FReply OnMouseButtonDown(const FGeometry&, const FPointerEvent&) override
+	{
+		return bAbsorbMouse ? FReply::Handled() : FReply::Unhandled();
+	}
+
+	virtual FReply OnMouseButtonUp(const FGeometry&, const FPointerEvent&) override
+	{
+		return bAbsorbMouse ? FReply::Handled() : FReply::Unhandled();
+	}
+
+	virtual FReply OnMouseButtonDoubleClick(const FGeometry&, const FPointerEvent&) override
+	{
+		return bAbsorbMouse ? FReply::Handled() : FReply::Unhandled();
+	}
+
+	virtual FReply OnMouseMove(const FGeometry&, const FPointerEvent&) override
+	{
+		return bAbsorbMouse ? FReply::Handled() : FReply::Unhandled();
+	}
+
+	virtual FReply OnMouseWheel(const FGeometry& MyGeometry, const FPointerEvent& MouseEvent) override
+	{
+		return bAbsorbMouse ? FReply::Handled() : FReply::Unhandled();
+	}
+	
+private:
+	TSharedPtr<SImStackBox> Container;
+	bool bAbsorbMouse = true;
+};
diff --git a/Engine/Plugins/Experimental/SlateIM/Source/SlateIM/Private/Containers/SImContextMenuAnchor.cpp b/Engine/Plugins/Experimental/SlateIM/Source/SlateIM/Private/Containers/SImContextMenuAnchor.cpp
new file mode 100644
--- /dev/null
+++ b/Engine/Plugins/Experimental/SlateIM/Source/SlateIM/Private/Containers/SImContextMenuAnchor.cpp
@@ -0,0 +1,353 @@
+﻿// Copyright Epic Games, Inc. All Rights Reserved.
+
+#include "SImContextMenuAnchor.h"
+
+#include "Framework/Application/MenuStack.h"
+#include "Framework/Application/SlateApplication.h"
+#include "Framework/MultiBox/MultiBoxBuilder.h"
+#include "Layout/WidgetPath.h"
+#include "Misc/SlateIMSlotData.h"
+
+
+SLATE_IMPLEMENT_WIDGET(SImContextMenuAnchor)
+
+void SImContextMenuAnchor::PrivateRegisterAttributes(FSlateAttributeInitializer& AttributeInitializer)
+{
+}
+
+void SImContextMenuAnchor::Construct(const FArguments& InArgs)
+{
+}
+
+int32 SImContextMenuAnchor::GetNumChildren()
+{
+	return GetChildren()->Num();
+}
+
+FSlateIMChild SImContextMenuAnchor::GetChild(int32 Index)
+{
+	if (Index >= 0 && Index < GetNumChildren())
+	{
+		TSharedRef<SWidget> Child = GetChildren()->GetChildAt(Index);
+
+		if (Child->GetWidgetClass().GetWidgetType() == SBox::StaticWidgetClass().GetWidgetType())
+		{
+			TSharedRef<SBox> SlotBox = StaticCastSharedRef<SBox>(Child);
+			Child = (SlotBox->GetChildren() && SlotBox->GetChildren()->Num() > 0) ? SlotBox->GetChildren()->GetChildAt(0) : SNullWidget::NullWidget;
+		}
+		
+		return Child;
+	}
+
+	return nullptr;
+}
+
+void SImContextMenuAnchor::RemoveUnusedChildren(int32 LastUsedChildIndex)
+{
+	ChildSlot
+	[
+		SNullWidget::NullWidget
+	];
+}
+
+void SImContextMenuAnchor::UpdateChild(FSlateIMChild Child, int32 Index, const FSlateIMSlotData& AlignmentData)
+{
+	ChildSlot
+	.Padding(AlignmentData.Padding)
+	.HAlign(AlignmentData.HorizontalAlignment)
+	.VAlign(AlignmentData.VerticalAlignment)
+	[
+		SNew(SBox)
+		.MinDesiredWidth(AlignmentData.MinWidth > 0 ? AlignmentData.MinWidth : FOptionalSize())
+		.MinDesiredHeight(AlignmentData.MinHeight > 0 ? AlignmentData.MinHeight : FOptionalSize())
+		.MaxDesiredWidth(AlignmentData.MaxWidth > 0 ? AlignmentData.MaxWidth : FOptionalSize())
+		.MaxDesiredHeight(AlignmentData.MaxHeight > 0 ? AlignmentData.MaxHeight : FOptionalSize())
+		[
+			Child.GetWidgetRef()
+		]
+	];
+}
+
+void SImContextMenuAnchor::Begin()
+{
+	bIsDirty = false;
+	CurrentMenuIndex = 0;
+	CurrentSubMenuLevel = 0;
+}
+
+void SImContextMenuAnchor::End()
+{
+	bIsDirty |= CurrentMenuIndex != MenuHashes.Num();
+
+	if (bIsDirty)
+	{
+		bIsDirty = false;
+
+		int32 NumMenuItems = CurrentMenuIndex;
+		int32 MenuHashCount = MenuHashes.Num();
+
+		if (MenuHashCount > NumMenuItems)
+		{
+			const int32 NumToRemove = MenuHashCount - NumMenuItems;
+
+			MenuHashes.RemoveAt(CurrentMenuIndex, NumToRemove);
+			CheckStates.RemoveAt(CurrentMenuIndex, NumToRemove);
+		}
+
+		SubMenuWidgets.Empty();
+
+		int32 StartIndex = 0;
+		const int32 MenuLevel = 0;
+		MenuWidget = BuildMenu_Recursive(StartIndex, MenuLevel);
+	}
+
+
+	MenuDataList.Empty();
+	MenuStringList.Empty();
+	ActivatedIndices.Empty();
+}
+
+bool SImContextMenuAnchor::AddMenuButton(const FStringView& RowText, const FStringView& ToolTipText)
+{
+	bool bCurrentStateIgnored = false;
+	return AddMenuInternal(RowText, ToolTipText, bCurrentStateIgnored, EMenuType::Button);
+}
+
+bool SImContextMenuAnchor::AddMenuCheckButton(const FStringView& RowText, bool& InOutCurrentState, const FStringView& ToolTipText)
+{
+	return AddMenuInternal(RowText, ToolTipText, InOutCurrentState,  EMenuType::Check);
+}
+
+bool SImContextMenuAnchor::AddMenuToggleButton(const FStringView& RowText, bool& InOutCurrentState, const FStringView& ToolTipText)
+{
+	return AddMenuInternal(RowText, ToolTipText, InOutCurrentState, EMenuType::Toggle);
+}
+
+
+void SImContextMenuAnchor::AddMenuSeparator()
+{
+	bool bCurrentStateIgnored = false;
+	AddMenuInternal(nullptr, nullptr, bCurrentStateIgnored, EMenuType::Separator);
+}
+
+void SImContextMenuAnchor::AddMenuSection(const FStringView& SectionText)
+{
+	bool bCurrentStateIgnored = false;
+	AddMenuInternal(SectionText, nullptr, bCurrentStateIgnored, EMenuType::Section);
+}
+
+void SImContextMenuAnchor::BeginSubMenu(const FStringView& SectionText)
+{
+	bool bCurrentStateIgnored = false;
+	AddMenuInternal(SectionText, nullptr, bCurrentStateIgnored, EMenuType::SubMenu);
+		
+	++CurrentSubMenuLevel;
+}
+
+void SImContextMenuAnchor::EndSubMenu()
+{
+	--CurrentSubMenuLevel;
+	checkf(CurrentSubMenuLevel >= 0, TEXT("Too many calls to EndSubMenu. This means you called EndSubMenu more than BeginSubMenu"));
+}
+
+FReply SImContextMenuAnchor::OnMouseButtonUp(const FGeometry& MyGeometry, const FPointerEvent& MouseEvent)
+{
+	if (MenuWidget)
+	{
+		FWidgetPath WidgetPath = MouseEvent.GetEventPath() != nullptr ? *MouseEvent.GetEventPath() : FWidgetPath();
+		OpenedMenu = FSlateApplication::Get().PushMenu(AsShared(), WidgetPath, MenuWidget.ToSharedRef(), MouseEvent.GetScreenSpacePosition(), FPopupTransitionEffect(FPopupTransitionEffect::ContextMenu));
+
+		return FReply::Handled();
+	}
+
+	return FReply::Unhandled();
+}
+
+bool SImContextMenuAnchor::AddMenuInternal(const FStringView& RowText, const FStringView& ToolTipText, bool& InOutCurrentState, EMenuType MenuType)
+{
+	bool bIsActivated = false;
+	uint64 RowHash = 0;
+	if (RowText.Len() > 0)
+	{
+		RowHash = CityHash64(reinterpret_cast<const char*>(RowText.GetData()), RowText.Len() * sizeof(TCHAR));
+	}
+
+	if (ToolTipText.Len() > 0)
+	{
+		RowHash = CityHash64WithSeed(reinterpret_cast<const char*>(ToolTipText.GetData()), ToolTipText.Len() * sizeof(TCHAR), RowHash);
+	}
+
+	FMenuItemData NewMenuItem;
+
+	if (RowText.Len() > 0)
+	{
+		NewMenuItem.TextOffset = MenuStringList.AddUninitialized(RowText.Len());
+		NewMenuItem.TextSize = RowText.Len();
+		RowText.CopyString(&MenuStringList[NewMenuItem.TextOffset], RowText.Len(), 0);
+	}
+	else
+	{
+		NewMenuItem.TextOffset = INDEX_NONE;
+		NewMenuItem.TextSize = 0;
+	}
+
+	if (ToolTipText.Len() > 0)
+	{
+		NewMenuItem.ToolTipOffset = MenuStringList.AddUninitialized(ToolTipText.Len());
+		NewMenuItem.ToolTipSize = ToolTipText.Len();
+		ToolTipText.CopyString(&MenuStringList[NewMenuItem.ToolTipOffset], ToolTipText.Len(), 0);
+	}
+	else
+	{
+		NewMenuItem.ToolTipOffset = INDEX_NONE;
+		NewMenuItem.ToolTipSize = 0;
+	}
+
+
+	NewMenuItem.Type = MenuType;
+	NewMenuItem.SubMenuLevel = CurrentSubMenuLevel;
+
+	RowHash = CityHash64WithSeed(reinterpret_cast<const char*>(&NewMenuItem), sizeof(NewMenuItem), RowHash);
+
+	MenuDataList.Add(MoveTemp(NewMenuItem));
+
+	bIsDirty |= MenuHashes.Num() > CurrentMenuIndex ? MenuHashes[CurrentMenuIndex] != RowHash : true;
+	if (bIsDirty)
+	{
+		if (MenuHashes.Num() > CurrentMenuIndex)
+		{
+			MenuHashes[CurrentMenuIndex] = RowHash;
+			CheckStates[CurrentMenuIndex] = InOutCurrentState;
+		}
+		else
+		{
+			MenuHashes.Add(RowHash);
+			CheckStates.Add(InOutCurrentState);
+			check(MenuHashes.Num() - 1 == CurrentMenuIndex);
+		}
+	}
+	else
+	{
+		bIsActivated = ActivatedIndices.Contains(CurrentMenuIndex);
+
+		bool bIsCheckType = MenuType == EMenuType::Check || MenuType == EMenuType::Toggle;
+		if (bIsActivated && bIsCheckType)
+		{
+			InOutCurrentState = CheckStates[CurrentMenuIndex];
+		}
+		else if (bIsCheckType)
+		{
+			CheckStates[CurrentMenuIndex] = InOutCurrentState;
+		}
+	}
+
+	++CurrentMenuIndex;
+
+	return bIsActivated;
+}
+
+void SImContextMenuAnchor::OnMenuItemExecuted(int32 MenuIndex, bool bIsCheck)
+{
+	ActivatedIndices.Add(MenuIndex);
+
+	if (bIsCheck)
+	{
+		CheckStates[MenuIndex] = !CheckStates[MenuIndex];
+	}
+}
+
+ECheckBoxState SImContextMenuAnchor::OnGetMenuItemCheckState(int32 MenuIndex)
+{
+	return CheckStates[MenuIndex] ? ECheckBoxState::Checked : ECheckBoxState::Unchecked;
+}
+
+TSharedRef<SWidget> SImContextMenuAnchor::OnGetSubMenuContent(int32 SubMenuIndex)
+{
+	//MenuBuilder.AddWidget(SubMenuWidgets[SubMenuIndex], FText::GetEmpty(), false, false);
+
+	return SubMenuWidgets[SubMenuIndex];
+}
+
+TSharedRef<SWidget> SImContextMenuAnchor::BuildMenu_Recursive(int32& InOutMenuIndex, int32 MenuLevel)
+{
+	int32 NumMenuItems = CurrentMenuIndex;
+
+	FMenuBuilder MenuBuilder(true, nullptr, nullptr, false, &FCoreStyle::Get(), false);
+
+	bool bHasCurrentSection = false;
+
+	while(InOutMenuIndex < NumMenuItems)
+	{
+		const FMenuItemData& MenuItemData = MenuDataList[InOutMenuIndex];
+
+		if (MenuItemData.SubMenuLevel < MenuLevel)
+		{
+			break;
+		}
+
+		FStringView MenuText;
+		FStringView ToolTipText;
+		if (MenuItemData.TextOffset != INDEX_NONE)
+		{
+			MenuText = FStringView(&MenuStringList[MenuItemData.TextOffset], MenuItemData.TextSize);
+		}
+
+		if (MenuItemData.ToolTipOffset != INDEX_NONE)
+		{
+			ToolTipText = FStringView(&MenuStringList[MenuItemData.ToolTipOffset], MenuItemData.ToolTipSize);
+		}
+
+		switch (MenuItemData.Type)
+		{
+		case EMenuType::Button:
+		case EMenuType::Check:
+		case EMenuType::Toggle:
+		{
+			bool bIsCheckType = MenuItemData.Type == EMenuType::Check || MenuItemData.Type == EMenuType::Toggle;
+
+			FUIAction Action(FExecuteAction::CreateSP(this, &SImContextMenuAnchor::OnMenuItemExecuted, InOutMenuIndex, bIsCheckType));
+
+			EUserInterfaceActionType MenuBuilderType = EUserInterfaceActionType::Button;
+
+			if(bIsCheckType)
+			{
+				MenuBuilderType = MenuItemData.Type == EMenuType::Check ? EUserInterfaceActionType::Check : EUserInterfaceActionType::ToggleButton;
+				Action.GetActionCheckState = FGetActionCheckState::CreateSP(this, &SImContextMenuAnchor::OnGetMenuItemCheckState, InOutMenuIndex);
+			}
+
+			MenuBuilder.AddMenuEntry(FText::FromStringView(MenuText), FText::FromStringView(ToolTipText), FSlateIcon(), Action, NAME_None, MenuBuilderType);
+		}
+		break;
+
+		case EMenuType::Separator:
+			MenuBuilder.AddMenuSeparator();
+			break;
+		case EMenuType::Section:
+			if (bHasCurrentSection)
+			{
+				MenuBuilder.EndSection();
+			}
+
+			bHasCurrentSection = true;
+			MenuBuilder.BeginSection(NAME_None, FText::FromStringView(MenuText));
+			break;
+
+		case EMenuType::SubMenu:
+			++InOutMenuIndex;
+			TSharedRef<SWidget> SubMenuWidget = BuildMenu_Recursive(InOutMenuIndex, MenuLevel+1);
+			int32 SubMenuIndex = SubMenuWidgets.Add(SubMenuWidget);
+			MenuBuilder.AddWrapperSubMenu(FText::FromStringView(MenuText), FText::FromStringView(ToolTipText), FOnGetContent::CreateSP(this, &SImContextMenuAnchor::OnGetSubMenuContent, SubMenuIndex), FSlateIcon());
+	
+			break;
+		}
+
+		++InOutMenuIndex;
+	}
+
+	if (bHasCurrentSection)
+	{
+		MenuBuilder.EndSection();
+	}
+
+	return MenuBuilder.MakeWidget();
+}
diff --git a/Engine/Plugins/Experimental/SlateIM/Source/SlateIM/Private/Containers/SImContextMenuAnchor.h b/Engine/Plugins/Experimental/SlateIM/Source/SlateIM/Private/Containers/SImContextMenuAnchor.h
new file mode 100644
--- /dev/null
+++ b/Engine/Plugins/Experimental/SlateIM/Source/SlateIM/Private/Containers/SImContextMenuAnchor.h
@@ -0,0 +1,88 @@
+﻿// Copyright Epic Games, Inc. All Rights Reserved.
+
+#pragma once
+
+#include "ISlateIMContainer.h"
+#include "Widgets/SCompoundWidget.h"
+
+class IMenu;
+
+enum class EMenuType : uint8
+{
+	Button,
+	Check,
+	Toggle,
+	Separator,
+	Section,
+	SubMenu,
+};
+
+class SImContextMenuAnchor : public SCompoundWidget, public ISlateIMContainer
+{
+	SLATE_DECLARE_WIDGET(SImContextMenuAnchor, SCompoundWidget)
+	SLATE_IM_TYPE_DATA(SImContextMenuAnchor, ISlateIMContainer)
+	
+public:
+	SLATE_BEGIN_ARGS(SImContextMenuAnchor) {}
+	SLATE_END_ARGS()
+
+	void Construct(const FArguments& InArgs);
+
+	virtual int32 GetNumChildren() override;
+	virtual FSlateIMChild GetChild(int32 Index) override;
+
+	virtual FSlateIMChild GetContainer() override
+	{
+		return AsShared();
+	}
+
+	virtual void RemoveUnusedChildren(int32 LastUsedChildIndex) override;
+
+	virtual void UpdateChild(FSlateIMChild Child, int32 Index, const FSlateIMSlotData& AlignmentData) override;
+
+	void Begin();
+	void End();
+
+	bool AddMenuButton(const FStringView& RowText, const FStringView& ToolTipText);
+	bool AddMenuCheckButton(const FStringView& RowText, bool& InOutCurrentState, const FStringView& ToolTipText);
+	bool AddMenuToggleButton(const FStringView& RowText, bool& InOutCurrentState, const FStringView& ToolTipText);
+	void AddMenuSeparator();
+	void AddMenuSection(const FStringView& SectionText);
+	void BeginSubMenu(const FStringView& SectionText);
+	void EndSubMenu();
+	bool IsMenuOpen() const { return OpenedMenu.IsValid();  }
+
+	virtual FReply OnMouseButtonUp(const FGeometry& MyGeometry, const FPointerEvent& MouseEvent) override;
+	
+private:
+	bool AddMenuInternal(const FStringView& RowText, const FStringView& ToolTipText, bool& InOutCurrentState, EMenuType MenuType);
+	void OnMenuItemExecuted(int32 MenuIndex, bool bIsCheck);
+	ECheckBoxState OnGetMenuItemCheckState(int32 MenuIndex);
+	//void OnGetSubMenuContent(FMenuBuilder& MenuBuilder, int32 SubMenuIndex);
+	TSharedRef<SWidget> OnGetSubMenuContent(int32 SubMenuIndex);
+	TSharedRef<SWidget> BuildMenu_Recursive(int32& InOutMenuIndex, int32 MenuLevel);
+
+	struct FMenuItemData
+	{
+		int32 TextOffset;
+		int32 TextSize;
+		int32 ToolTipOffset;
+		int32 ToolTipSize;
+		int32 SubMenuLevel;
+		EMenuType Type;
+
+	};
+
+	TWeakPtr<IMenu> OpenedMenu;
+	TSharedPtr<SWidget> MenuWidget;
+	TArray<FMenuItemData, TMemStackAllocator<>> MenuDataList;
+	TArray<TCHAR, TMemStackAllocator<>> MenuStringList;
+
+	TArray<int32, TInlineAllocator<1>> ActivatedIndices;
+	TBitArray<> CheckStates;
+	TArray<TSharedRef<SWidget>> SubMenuWidgets;
+	TArray<uint64> MenuHashes;
+	int32 CurrentMenuIndex = 0;
+	int32 CurrentSubMenuLevel = 0;
+	bool bIsDirty = false;
+};
diff --git a/Engine/Plugins/Experimental/SlateIM/Source/SlateIM/Private/Containers/SImScrollBox.cpp b/Engine/Plugins/Experimental/SlateIM/Source/SlateIM/Private/Containers/SImScrollBox.cpp
new file mode 100644
--- /dev/null
+++ b/Engine/Plugins/Experimental/SlateIM/Source/SlateIM/Private/Containers/SImScrollBox.cpp
@@ -0,0 +1,84 @@
+﻿// Copyright Epic Games, Inc. All Rights Reserved.
+
+#include "SImScrollBox.h"
+#include "Misc/SlateIMSlotData.h"
+#include "Widgets/Layout/SBox.h"
+
+SLATE_IMPLEMENT_WIDGET(SImScrollBox)
+
+void SImScrollBox::PrivateRegisterAttributes(FSlateAttributeInitializer& AttributeInitializer)
+{
+}
+
+int32 SImScrollBox::GetNumChildren()
+{
+	return ScrollPanel->GetChildren()->Num();
+}
+
+FSlateIMChild SImScrollBox::GetChild(int32 Index)
+{
+	if (Index >= 0 && Index < GetNumChildren())
+	{
+		TSharedRef<SWidget> Child = ScrollPanel->GetChildren()->GetChildAt(Index);
+
+		if (Child->GetWidgetClass().GetWidgetType() == SBox::StaticWidgetClass().GetWidgetType())
+		{
+			TSharedRef<SBox> SlotBox = StaticCastSharedRef<SBox>(Child);
+			Child = (SlotBox->GetChildren() && SlotBox->GetChildren()->Num() > 0) ? SlotBox->GetChildren()->GetChildAt(0) : SNullWidget::NullWidget;
+		}
+		
+		return Child;
+	}
+
+	return nullptr;
+}
+
+void SImScrollBox::UpdateChild(FSlateIMChild Child, int32 Index, const FSlateIMSlotData& AlignmentData)
+{
+	SScrollBox::FSlot* Slot = nullptr;
+
+	if (IsValidSlotIndex(Index))
+	{
+		Slot = &GetSlot(Index);
+	}
+	else
+	{
+		AddSlot().Expose(Slot);
+	}
+
+	if (AlignmentData.bAutoSize)
+	{
+		Slot->SetSizeToAuto();
+	}
+	else
+	{
+		Slot->SetSizeToStretchContent(1.0f);
+	}
+
+	Slot->SetPadding(AlignmentData.Padding);
+	Slot->SetHorizontalAlignment(AlignmentData.HorizontalAlignment);
+	Slot->SetVerticalAlignment(AlignmentData.VerticalAlignment);
+	if (Slot->GetWidget()->GetWidgetClass().GetWidgetType() != SBox::StaticWidgetClass().GetWidgetType())
+	{
+		(*Slot)
+		[
+			SNew(SBox)
+			.MinDesiredWidth(AlignmentData.MinWidth > 0 ? AlignmentData.MinWidth : FOptionalSize())
+			.MinDesiredHeight(AlignmentData.MinHeight > 0 ? AlignmentData.MinHeight : FOptionalSize())
+			.MaxDesiredWidth(AlignmentData.MaxWidth > 0 ? AlignmentData.MaxWidth : FOptionalSize())
+			.MaxDesiredHeight(AlignmentData.MaxHeight > 0 ? AlignmentData.MaxHeight : FOptionalSize())
+			[
+				Child.GetWidgetRef()
+			]
+		];
+	}
+	else
+	{
+		TSharedRef<SBox> SlotBox = StaticCastSharedRef<SBox>(Slot->GetWidget());
+		SlotBox->SetMinDesiredWidth(AlignmentData.MinWidth > 0 ? AlignmentData.MinWidth : FOptionalSize());
+		SlotBox->SetMinDesiredHeight(AlignmentData.MinHeight > 0 ? AlignmentData.MinHeight : FOptionalSize());
+		SlotBox->SetMaxDesiredWidth(AlignmentData.MaxWidth > 0 ? AlignmentData.MaxWidth : FOptionalSize());
+		SlotBox->SetMaxDesiredHeight(AlignmentData.MaxHeight > 0 ? AlignmentData.MaxHeight : FOptionalSize());
+		SlotBox->SetContent(Child.GetWidgetRef());
+	}
+}
diff --git a/Engine/Plugins/Experimental/SlateIM/Source/SlateIM/Private/Containers/SImScrollBox.h b/Engine/Plugins/Experimental/SlateIM/Source/SlateIM/Private/Containers/SImScrollBox.h
new file mode 100644
--- /dev/null
+++ b/Engine/Plugins/Experimental/SlateIM/Source/SlateIM/Private/Containers/SImScrollBox.h
@@ -0,0 +1,41 @@
+﻿// Copyright Epic Games, Inc. All Rights Reserved.
+
+#pragma once
+
+#include "ISlateIMContainer.h"
+#include "Widgets/Layout/SScrollBox.h"
+
+class SImScrollBox : public SScrollBox, public ISlateIMContainer
+{
+	SLATE_DECLARE_WIDGET(SImScrollBox, SScrollBox)
+	SLATE_IM_TYPE_DATA(SImScrollBox, ISlateIMContainer)
+	
+public:
+	bool IsValidSlotIndex(int32 Index) const
+	{
+		return ScrollPanel->Children.IsValidIndex(Index);
+	}
+
+	void SetOnUserScrolled(const FOnUserScrolled& InHandler)
+	{
+		OnUserScrolled = InHandler;
+	}
+
+	virtual int32 GetNumChildren() override;
+	virtual FSlateIMChild GetChild(int32 Index) override;
+
+	virtual FSlateIMChild GetContainer() override
+	{
+		return AsShared();
+	}
+
+	virtual void RemoveUnusedChildren(int32 LastUsedChildIndex) override
+	{
+		for (int32 IndexToRemove = GetNumChildren() - 1; IndexToRemove > LastUsedChildIndex; --IndexToRemove)
+		{
+			ScrollPanel->Children.RemoveAt(IndexToRemove);
+		}
+	}
+
+	virtual void UpdateChild(FSlateIMChild Child, int32 Index, const FSlateIMSlotData& AlignmentData) override;
+};
diff --git a/Engine/Plugins/Experimental/SlateIM/Source/SlateIM/Private/Containers/SImStackBox.cpp b/Engine/Plugins/Experimental/SlateIM/Source/SlateIM/Private/Containers/SImStackBox.cpp
new file mode 100644
--- /dev/null
+++ b/Engine/Plugins/Experimental/SlateIM/Source/SlateIM/Private/Containers/SImStackBox.cpp
@@ -0,0 +1,93 @@
+﻿// Copyright Epic Games, Inc. All Rights Reserved.
+
+#include "SImStackBox.h"
+
+#include "Misc/SlateIMSlotData.h"
+#include "Widgets/Layout/SBox.h"
+
+SLATE_IMPLEMENT_WIDGET(SImStackBox)
+
+void SImStackBox::PrivateRegisterAttributes(FSlateAttributeInitializer&)
+{
+}
+
+int32 SImStackBox::GetNumChildren()
+{
+	return GetChildren()->Num();
+}
+
+FSlateIMChild SImStackBox::GetChild(int32 Index)
+{
+	if (Index >= 0 && Index < GetNumChildren())
+	{
+		TSharedRef<SWidget> Child = GetChildren()->GetChildAt(Index);
+
+		if (Child->GetWidgetClass().GetWidgetType() == SBox::StaticWidgetClass().GetWidgetType())
+		{
+			TSharedRef<SBox> SlotBox = StaticCastSharedRef<SBox>(Child);
+			Child = (SlotBox->GetChildren() && SlotBox->GetChildren()->Num() > 0) ? SlotBox->GetChildren()->GetChildAt(0) : SNullWidget::NullWidget;
+		}
+		
+		return Child;
+	}
+
+	return nullptr;
+}
+
+void SImStackBox::RemoveUnusedChildren(int32 LastUsedChildIndex)
+{
+	for (int32 IndexToRemove = GetChildren()->Num() - 1; IndexToRemove > LastUsedChildIndex; --IndexToRemove)
+	{
+		RemoveSlot(GetChildren()->GetChildAt(IndexToRemove));
+	}
+}
+
+void SImStackBox::UpdateChild(FSlateIMChild Child, int32 Index, const FSlateIMSlotData& AlignmentData)
+{
+	SStackBox::FSlot* Slot = nullptr;
+
+	if (IsValidSlotIndex(Index))
+	{
+		Slot = &GetSlot(Index);
+	}
+	else
+	{
+		AddSlot().Expose(Slot);
+	}
+
+	if (AlignmentData.bAutoSize)
+	{
+		Slot->SetSizeToAuto();
+	}
+	else
+	{
+		Slot->SetSizeToStretchContent(1.0f);
+	}
+
+	Slot->SetHorizontalAlignment(AlignmentData.HorizontalAlignment);
+	Slot->SetVerticalAlignment(AlignmentData.VerticalAlignment);
+	Slot->SetPadding(AlignmentData.Padding);
+	if (Slot->GetWidget()->GetWidgetClass().GetWidgetType() != SBox::StaticWidgetClass().GetWidgetType())
+	{
+		(*Slot)
+		[
+			SNew(SBox)
+			.MinDesiredWidth(AlignmentData.MinWidth > 0 ? AlignmentData.MinWidth : FOptionalSize())
+			.MinDesiredHeight(AlignmentData.MinHeight > 0 ? AlignmentData.MinHeight : FOptionalSize())
+			.MaxDesiredWidth(AlignmentData.MaxWidth > 0 ? AlignmentData.MaxWidth : FOptionalSize())
+			.MaxDesiredHeight(AlignmentData.MaxHeight > 0 ? AlignmentData.MaxHeight : FOptionalSize())
+			[
+				Child.GetWidgetRef()
+			]
+		];
+	}
+	else
+	{
+		TSharedRef<SBox> SlotBox = StaticCastSharedRef<SBox>(Slot->GetWidget());
+		SlotBox->SetMinDesiredWidth(AlignmentData.MinWidth > 0 ? AlignmentData.MinWidth : FOptionalSize());
+		SlotBox->SetMinDesiredHeight(AlignmentData.MinHeight > 0 ? AlignmentData.MinHeight : FOptionalSize());
+		SlotBox->SetMaxDesiredWidth(AlignmentData.MaxWidth > 0 ? AlignmentData.MaxWidth : FOptionalSize());
+		SlotBox->SetMaxDesiredHeight(AlignmentData.MaxHeight > 0 ? AlignmentData.MaxHeight : FOptionalSize());
+		SlotBox->SetContent(Child.GetWidgetRef());
+	}
+}
\ No newline at end of file
diff --git a/Engine/Plugins/Experimental/SlateIM/Source/SlateIM/Private/Containers/SImStackBox.h b/Engine/Plugins/Experimental/SlateIM/Source/SlateIM/Private/Containers/SImStackBox.h
new file mode 100644
--- /dev/null
+++ b/Engine/Plugins/Experimental/SlateIM/Source/SlateIM/Private/Containers/SImStackBox.h
@@ -0,0 +1,25 @@
+﻿// Copyright Epic Games, Inc. All Rights Reserved.
+
+#pragma once
+
+#include "ISlateIMContainer.h"
+#include "Widgets/SBoxPanel.h"
+
+class SImStackBox : public SStackBox, public ISlateIMContainer
+{
+	SLATE_DECLARE_WIDGET(SImStackBox, SStackBox)
+	SLATE_IM_TYPE_DATA(SImStackBox, ISlateIMContainer)
+	
+public:
+	virtual int32 GetNumChildren() override;
+	virtual FSlateIMChild GetChild(int32 Index) override;
+
+	virtual FSlateIMChild GetContainer() override
+	{
+		return AsShared();
+	}
+
+	virtual void RemoveUnusedChildren(int32 LastUsedChildIndex) override;
+
+	virtual void UpdateChild(FSlateIMChild Child, int32 Index, const FSlateIMSlotData& AlignmentData) override;
+};
diff --git a/Engine/Plugins/Experimental/SlateIM/Source/SlateIM/Private/Containers/SImTableView.cpp b/Engine/Plugins/Experimental/SlateIM/Source/SlateIM/Private/Containers/SImTableView.cpp
new file mode 100644
--- /dev/null
+++ b/Engine/Plugins/Experimental/SlateIM/Source/SlateIM/Private/Containers/SImTableView.cpp
@@ -0,0 +1,410 @@
+﻿// Copyright Epic Games, Inc. All Rights Reserved.
+
+#include "SImTableView.h"
+
+#include "Misc/SlateIMLogging.h"
+#include "Misc/SlateIMSlotData.h"
+
+SLATE_IMPLEMENT_WIDGET(SImTableRow)
+SLATE_IMPLEMENT_WIDGET(SImTableHeader)
+SLATE_IMPLEMENT_WIDGET(SImTableView)
+
+namespace SlateIM::Private
+{
+	const FName ExpanderColumnName = TEXT("Expander");
+	const FName BaseColumnName = TEXT("TableColumn");
+}
+
+
+void FSlateIMTableRow::RemoveUnusedChildren(int32 LastUsedChildIndex)
+{
+	Children.SetNum(FMath::Min(Children.Num(), LastUsedChildIndex + 1));
+}
+
+void FSlateIMTableRow::UpdateChild(FSlateIMChild Child, int32 Index, const FSlateIMSlotData& AlignmentData)
+{
+	if (TSharedPtr<FSlateIMTableRow> ChildRow = Child.GetChild<FSlateIMTableRow>())
+	{
+		ChildRow->UpdateColumnCount(ColumnCount);
+		ChildRow->SetOwningTable(OwningTable);
+	}
+
+	if (Children.IsValidIndex(Index))
+	{
+		Children[Index] = Child;
+	}
+	else
+	{
+		Children.Add(Child);
+	}
+
+	if (OwningTable)
+	{
+		OwningTable->MarkDirty();
+	}
+}
+
+void FSlateIMTableRow::GetChildRows(TArray<TSharedRef<FSlateIMTableRow>>& OutRows) const
+{
+	for (const FSlateIMChild& Child : Children)
+	{
+		if (TSharedPtr<FSlateIMTableRow> ChildRow = Child.GetChild<FSlateIMTableRow>())
+		{
+			OutRows.Add(ChildRow.ToSharedRef());
+		}
+	}
+}
+
+int32 FSlateIMTableRow::CountCellWidgetsUpToIndex(int32 Index) const
+{
+	int32 CellCount = 0;
+	
+	for (int32 i = 0; i < Children.Num() && i <= Index; ++i)
+	{
+		const FSlateIMChild& Child = Children[i];
+		if (TSharedPtr<SWidget> ChildRow = Child.GetWidget())
+		{
+			++CellCount;
+		}
+	}
+
+	return CellCount;
+}
+
+TSharedRef<SWidget> FSlateIMTableRow::GetCellWidget(int32 CellIndex) const
+{
+	int32 CellCount = 0;
+	for (const FSlateIMChild& Child : Children)
+	{
+		if (TSharedPtr<SWidget> ChildRow = Child.GetWidget())
+		{
+			if (CellIndex == CellCount)
+			{
+				return ChildRow.ToSharedRef();
+			}
+
+			++CellCount;
+		}
+	}
+
+	return SNullWidget::NullWidget;
+}
+
+void FSlateIMTableRow::UpdateColumnCount(int32 NewColumnCount)
+{
+	ColumnCount = NewColumnCount;
+
+	for (const FSlateIMChild& Child : Children)
+	{
+		if (TSharedPtr<FSlateIMTableRow> ChildRow = Child.GetChild<FSlateIMTableRow>())
+		{
+			ChildRow->UpdateColumnCount(ColumnCount);
+		}
+	}
+}
+
+void FSlateIMTableRow::SetOwningTable(const TSharedPtr<SImTableView>& InOwningTable)
+{
+	OwningTable = InOwningTable;
+
+	for (const FSlateIMChild& Child : Children)
+	{
+		if (TSharedPtr<FSlateIMTableRow> ChildRow = Child.GetChild<FSlateIMTableRow>())
+		{
+			ChildRow->SetOwningTable(OwningTable);
+		}
+	}
+}
+
+bool FSlateIMTableRow::IsExpanded()
+{
+	return OwningTable && OwningTable->IsItemExpanded(AsShared());
+}
+
+bool FSlateIMTableRow::HasChildRows() const
+{
+	for (const FSlateIMChild& Child : Children)
+	{
+		if (TSharedPtr<FSlateIMTableRow> ChildRow = Child.GetChild<FSlateIMTableRow>())
+		{
+			return true;
+		}
+	}
+
+	return false;
+}
+
+bool FSlateIMTableRow::ShouldDisplayExpander() const
+{
+	return OwningTable && OwningTable->IsTree();
+}
+
+bool FSlateIMTableRow::AreTableRowContentsRequired()
+{
+	if (OwningTable)
+	{
+		if (OwningTable->WidgetFromItem(AsShared()).IsValid())
+		{
+			// This row has a live widget so the contents are required
+			return true;
+		}
+
+		const float NumLiveWidgets = OwningTable->GetNumLiveWidgets();
+		const float CurrentScrollOffset = OwningTable->GetScrollOffset();
+
+		// Require content from rows within one "page" of the current scroll position in either direction
+		// This is probably a bigger window than necessary but should handle large scroll movements and should be small enough to not impact perf
+		const float MinNearlyLiveWidget = FMath::Max(0.f, CurrentScrollOffset - NumLiveWidgets);
+		const float MaxNearlyLiveWidget = CurrentScrollOffset + (2.f * NumLiveWidgets);
+		const float ThisRowIndex = OwningTable->GetRowLinearizedIndex(AsShared());
+		if (ThisRowIndex >= MinNearlyLiveWidget && ThisRowIndex <= MaxNearlyLiveWidget)
+		{
+			return true;
+		}
+	}
+	
+	return false;
+}
+
+void SImTableRow::PrivateRegisterAttributes(FSlateAttributeInitializer&)
+{
+}
+
+void SImTableRow::Construct(const FArguments& InArgs, const TSharedRef<SImTableView>& InOwnerTableView, const TSharedRef<FSlateIMTableRow>& InTableRow)
+{
+	TableRow = InTableRow;
+	SMultiColumnTableRow<TSharedRef<FSlateIMTableRow>>::Construct(FSuperRowType::FArguments(), InOwnerTableView);
+}
+
+TSharedRef<SWidget> SImTableRow::GenerateWidgetForColumn(const FName& ColumnName)
+{
+	if (ColumnName == SlateIM::Private::ExpanderColumnName)
+	{
+		return SAssignNew(ExpanderArrowWidget, SExpanderArrow, SharedThis(this))
+			.StyleSet(ExpanderStyleSet)
+			.ShouldDrawWires(true);
+	}
+	else if (TableRow)
+	{
+		const int32 ColumnIndex = ColumnName.GetNumber();
+		if (ColumnIndex == 0 && TableRow->ShouldDisplayExpander())
+		{
+			return SNew(SHorizontalBox)
+				+ SHorizontalBox::Slot()
+				.AutoWidth()
+				.HAlign(HAlign_Right)
+				.VAlign(VAlign_Fill)
+				[
+					SAssignNew(ExpanderArrowWidget, SExpanderArrow, SharedThis(this))
+					.StyleSet(ExpanderStyleSet)
+					.ShouldDrawWires(true)
+				]
+				+ SHorizontalBox::Slot()
+				.FillWidth(1)
+				[
+					TableRow->GetCellWidget(ColumnIndex)
+				];
+		}
+		else
+		{
+			return TableRow->GetCellWidget(ColumnIndex);
+		}
+	}
+
+	return SNullWidget::NullWidget;
+}
+
+void SImTableHeader::PrivateRegisterAttributes(FSlateAttributeInitializer&)
+{
+}
+
+void SImTableView::PrivateRegisterAttributes(FSlateAttributeInitializer&)
+{
+}
+
+void SImTableView::Construct(const FArguments& InArgs)
+{
+	Super::Construct(
+		FArguments(InArgs)
+		.TreeItemsSource(&TableRows)
+		.HeaderRow(SAssignNew(Header, SImTableHeader))
+		.OnGenerateRow(this, &SImTableView::GenerateRow)
+		.OnGetChildren(this, &SImTableView::GatherChildren));
+}
+
+FSlateIMChild SImTableView::GetChild(int32 Index)
+{
+	return TableRows.IsValidIndex(Index) ? FSlateIMChild(TableRows[Index]) : FSlateIMChild(nullptr);
+}
+
+void SImTableView::RemoveUnusedChildren(int32 LastUsedChildIndex)
+{
+	if (TableRows.Num() > (LastUsedChildIndex + 1))
+	{
+		bIsDirty = true;
+		TableRows.SetNum(LastUsedChildIndex + 1);
+	}
+}
+
+void SImTableView::UpdateChild(FSlateIMChild Child, int32 Index, const FSlateIMSlotData& AlignmentData)
+{
+	TSharedPtr<FSlateIMTableRow> Row = Child.GetChild<FSlateIMTableRow>();
+	if (!ensureMsgf(Row, TEXT("Tables can only hold Rows/Cells")))
+	{
+		return;
+	}
+
+	Row->SetOwningTable(SharedThis(this));
+
+	if (TableRows.IsValidIndex(Index))
+	{
+		if (TableRows[Index] != Row)
+		{
+			bIsDirty = true;
+			TableRows[Index] = Row.ToSharedRef();
+		}
+	}
+	else
+	{
+		bIsDirty = true;
+		TableRows.Add(Row.ToSharedRef());
+	}
+
+	if (Row->GetColumnCount() != ColumnCount)
+	{
+		bIsDirty = true;
+		Row->UpdateColumnCount(ColumnCount);
+	}
+}
+
+float SImTableView::GetNumLiveWidgets() const
+{
+	// NumLiveWidgets can be 0 during updates, so we use the last valid value we had
+	const int32 NumLiveWidgets = STreeView<TSharedRef<FSlateIMTableRow>>::GetNumLiveWidgets();
+	if (NumLiveWidgets > 0)
+	{
+		CachedNumLiveWidgets = NumLiveWidgets;
+	}
+
+	return CachedNumLiveWidgets;
+}
+
+void SImTableView::AddColumn(const FStringView& ColumnLabel, const FStringView& ColumnToolTip, const FSlateIMSlotData& SlotData)
+{
+	// Only add columns that don't already exist
+	if (Header)
+	{
+		// TODO - Handle changes to SlotData after construction
+		const int32 MaxColumnIndex = Header->GetColumns().Num();
+		if (ColumnCount >= MaxColumnIndex)
+		{
+			FName ColumnName = SlateIM::Private::BaseColumnName;
+			ColumnName.SetNumber(ColumnCount);
+
+			AddColumn_Internal(ColumnName, ColumnToolTip, SlotData, ColumnLabel);
+		}
+	}
+	
+	++ColumnCount;
+}
+
+void SImTableView::BeginTableUpdates()
+{
+	ColumnCount = 0;
+}
+
+void SImTableView::EndTableUpdates()
+{
+	if (bIsDirty)
+	{
+		bIsDirty = false;
+		RequestListRefresh();
+	}
+}
+
+void SImTableView::BeginTableContent()
+{
+	if (ColumnCount == 0)
+	{
+		// TODO - Do we force create a column here or can we hide the header and automatically treat everything like a list?
+	}
+	UpdateColumns();
+}
+
+void SImTableView::UpdateColumns()
+{
+	if (Header)
+	{
+		FName ColumnName = SlateIM::Private::BaseColumnName;
+		const int32 MaxColumnIndex = Header->GetColumns().Num();
+		for (int32 ColumnIndex = ColumnCount; ColumnIndex < MaxColumnIndex; ++ColumnIndex)
+		{
+			ColumnName.SetNumber(ColumnIndex);
+			Header->RemoveColumn(ColumnName);
+		}
+	}
+}
+
+bool SImTableView::IsTree() const
+{
+	for (const TSharedRef<FSlateIMTableRow>& TableRow : TableRows)
+	{
+		if (TableRow->HasChildRows())
+		{
+			return true;
+		}
+	}
+
+	return false;
+}
+
+TSharedRef<ITableRow> SImTableView::GenerateRow(TSharedRef<FSlateIMTableRow> InTableRow, const TSharedRef<STableViewBase>& OwnerTable)
+{
+	return SNew(SImTableRow, SharedThis(this), InTableRow);
+}
+
+void SImTableView::GatherChildren(TSharedRef<FSlateIMTableRow> Row, TArray<TSharedRef<FSlateIMTableRow>>& OutChildren) const
+{
+	Row->GetChildRows(OutChildren);
+}
+
+int32 SImTableView::GetRowLinearizedIndex(const TSharedRef<FSlateIMTableRow>& Row) const
+{
+	return LinearizedItems.IndexOfByKey(Row);
+}
+
+void SImTableView::AddColumn_Internal(const FName& ColumnId, const FStringView& ColumnToolTip, const FSlateIMSlotData& SlotData, const FStringView& ColumnLabel)
+{
+	SHeaderRow::FColumn::FArguments Args = SHeaderRow::FColumn::FArguments()
+		.ColumnId(ColumnId)
+		.DefaultLabel(FText::FromStringView(ColumnLabel))
+		.DefaultTooltip(FText::FromStringView(ColumnToolTip))
+		.HAlignHeader(SlotData.HorizontalAlignment)
+		.VAlignHeader(SlotData.VerticalAlignment);
+
+	if (SlotData.bAutoSize && SlotData.MinWidth > 0)
+	{
+		if (SlotData.MaxWidth > 0 && SlotData.MaxWidth <= SlotData.MinWidth)
+		{
+			Args.FixedWidth(SlotData.MinWidth);
+		}
+		else
+		{
+			Args.ManualWidth(SlotData.MinWidth);
+		}
+	}
+	else
+	{
+		if (SlotData.MinWidth > 0)
+		{
+			Args.FillSized(SlotData.MinWidth);
+		}
+		else
+		{
+			Args.FillWidth(1.f);
+		}
+	}
+		
+	Header->AddColumn(Args);
+	bIsDirty = true;
+}
diff --git a/Engine/Plugins/Experimental/SlateIM/Source/SlateIM/Private/Containers/SImTableView.h b/Engine/Plugins/Experimental/SlateIM/Source/SlateIM/Private/Containers/SImTableView.h
new file mode 100644
--- /dev/null
+++ b/Engine/Plugins/Experimental/SlateIM/Source/SlateIM/Private/Containers/SImTableView.h
@@ -0,0 +1,117 @@
+﻿// Copyright Epic Games, Inc. All Rights Reserved.
+
+#pragma once
+
+#include "ISlateIMContainer.h"
+#include "Misc/ISlateIMChild.h"
+#include "Widgets/Views/STreeView.h"
+#include "Widgets/Views/STableRow.h"
+
+class SImTableView;
+
+class FSlateIMTableRow : public ISlateIMContainer, public ISlateIMChild, public TSharedFromThis<FSlateIMTableRow>
+{
+	SLATE_IM_TYPE_DATA(FSlateIMTableRow, ISlateIMChild)
+	// Deadly diamond of death, but it's okay since ISlateIMContainer and ISlateIMChild don't have implementations for ISlateIMTypeChecking
+	// SLATE_IM_TYPE_DATA(FSlateIMTableRow, ISlateIMContainer)
+
+public:
+	virtual FString GetDebugName() override { return GetTypeId().ToString(); }
+	virtual int32 GetNumChildren() override { return Children.Num(); }
+	virtual FSlateIMChild GetChild(int32 Index) override { return Children.IsValidIndex(Index) ? Children[Index] : nullptr; }
+	virtual FSlateIMChild GetContainer() override { return TSharedRef<ISlateIMChild>(AsShared()); }
+	virtual void RemoveUnusedChildren(int32 LastUsedChildIndex) override;
+	virtual void UpdateChild(FSlateIMChild Child, int32 Index, const FSlateIMSlotData& AlignmentData) override;
+
+	void GetChildRows(TArray<TSharedRef<FSlateIMTableRow>>& OutRows) const;
+
+	int32 CountCellWidgetsUpToIndex(int32 Index) const;
+
+	TSharedRef<SWidget> GetCellWidget(int32 CellIndex) const;
+
+	int32 GetColumnCount() const { return ColumnCount; }
+	void UpdateColumnCount(int32 NewColumnCount);
+
+	void SetOwningTable(const TSharedPtr<SImTableView>& InOwningTable);
+
+	bool IsExpanded();
+	bool HasChildRows() const;
+	bool ShouldDisplayExpander() const;
+	bool AreTableRowContentsRequired();
+
+private:
+	// A mix of cell widgets and child rows
+	TArray<FSlateIMChild> Children;
+	int32 ColumnCount = 0;
+	TSharedPtr<SImTableView> OwningTable;
+};
+
+class SImTableRow : public SMultiColumnTableRow<TSharedRef<FSlateIMTableRow>>
+{
+	SLATE_DECLARE_WIDGET(SImTableRow, SMultiColumnTableRow<TSharedRef<FSlateIMTableRow>>)
+	
+public:
+	SLATE_BEGIN_ARGS(SImTableRow) { }
+	SLATE_END_ARGS()
+	
+	void Construct(const FArguments& InArgs, const TSharedRef<SImTableView>& InOwnerTableView, const TSharedRef<FSlateIMTableRow>& InTableRow);
+	
+	virtual TSharedRef<SWidget> GenerateWidgetForColumn(const FName& ColumnName) override;
+	
+private:
+	TSharedPtr<FSlateIMTableRow> TableRow;
+};
+
+class SImTableHeader : public SHeaderRow
+{
+	SLATE_DECLARE_WIDGET(SImTableHeader, SHeaderRow)
+};
+
+class SImTableView : public STreeView<TSharedRef<FSlateIMTableRow>>, public ISlateIMContainer
+{
+	SLATE_DECLARE_WIDGET(SImTableView, STreeView<TSharedRef<FSlateIMTableRow>>)
+	SLATE_IM_TYPE_DATA(SImTableView, ISlateIMContainer)
+	
+public:
+	void Construct(const FArguments& InArgs);
+
+	virtual int32 GetNumChildren() override { return TableRows.Num(); }
+	virtual FSlateIMChild GetChild(int32 Index) override;
+	virtual FSlateIMChild GetContainer() override { return AsShared(); }
+	virtual void RemoveUnusedChildren(int32 LastUsedChildIndex) override;
+	virtual void UpdateChild(FSlateIMChild Child, int32 Index, const FSlateIMSlotData& AlignmentData) override;
+	
+	virtual float GetNumLiveWidgets() const override;
+
+	void AddColumn(const FStringView& ColumnLabel, const FStringView& ColumnToolTip, const FSlateIMSlotData& SlotData);
+
+	void BeginTableUpdates();
+	void EndTableUpdates();
+
+	void BeginTableContent();
+
+	void UpdateColumns();
+
+	// A table is a tree if any of its rows has children
+	bool IsTree() const;
+
+	void MarkDirty() { bIsDirty = true; }
+
+	TSharedRef<ITableRow> GenerateRow(TSharedRef<FSlateIMTableRow> InTableRow, const TSharedRef<STableViewBase>& OwnerTable);
+
+	void GatherChildren(TSharedRef<FSlateIMTableRow> Row, TArray<TSharedRef<FSlateIMTableRow>>& OutChildren) const;
+
+	TSharedPtr<FSlateIMTableRow> GetRow(int32 Index) { return TableRows.IsValidIndex(Index) ? TableRows[Index].ToSharedPtr() : nullptr; }
+
+	int32 GetRowLinearizedIndex(const TSharedRef<FSlateIMTableRow>& Row) const;
+
+private:
+	void AddColumn_Internal(const FName& ColumnId, const FStringView& ColumnToolTip, const FSlateIMSlotData& SlotData, const FStringView& ColumnLabel = FStringView());
+
+	TArray<TSharedRef<FSlateIMTableRow>> TableRows;
+	TSharedPtr<SImTableHeader> Header;
+	int32 ColumnCount = 0;
+	bool bIsDirty = false;
+
+	mutable int32 CachedNumLiveWidgets = 0;
+};
diff --git a/Engine/Plugins/Experimental/SlateIM/Source/SlateIM/Private/Containers/SImWrapBox.cpp b/Engine/Plugins/Experimental/SlateIM/Source/SlateIM/Private/Containers/SImWrapBox.cpp
new file mode 100644
--- /dev/null
+++ b/Engine/Plugins/Experimental/SlateIM/Source/SlateIM/Private/Containers/SImWrapBox.cpp
@@ -0,0 +1,84 @@
+﻿// Copyright Epic Games, Inc. All Rights Reserved.
+
+#include "SImWrapBox.h"
+
+#include "Misc/SlateIMSlotData.h"
+#include "Widgets/Layout/SBox.h"
+
+SLATE_IMPLEMENT_WIDGET(SImWrapBox)
+
+void SImWrapBox::PrivateRegisterAttributes(FSlateAttributeInitializer& AttributeInitializer)
+{
+}
+
+int32 SImWrapBox::GetNumChildren()
+{
+	return GetChildren()->Num();
+}
+
+FSlateIMChild SImWrapBox::GetChild(int32 Index)
+{
+	if (Index >= 0 && Index < GetNumChildren())
+	{
+		TSharedRef<SWidget> Child = GetChildren()->GetChildAt(Index);
+
+		if (Child->GetWidgetClass().GetWidgetType() == SBox::StaticWidgetClass().GetWidgetType())
+		{
+			TSharedRef<SBox> SlotBox = StaticCastSharedRef<SBox>(Child);
+			Child = (SlotBox->GetChildren() && SlotBox->GetChildren()->Num() > 0) ? SlotBox->GetChildren()->GetChildAt(0) : SNullWidget::NullWidget;
+		}
+		
+		return Child;
+	}
+
+	return nullptr;
+}
+
+void SImWrapBox::RemoveUnusedChildren(int32 LastUsedChildIndex)
+{
+	for (int32 IndexToRemove = GetChildren()->Num() - 1; IndexToRemove > LastUsedChildIndex; --IndexToRemove)
+	{
+		RemoveSlot(GetChildren()->GetChildAt(IndexToRemove));
+	}
+}
+
+void SImWrapBox::UpdateChild(FSlateIMChild Child, int32 Index, const FSlateIMSlotData& AlignmentData)
+{
+	SWrapBox::FSlot* Slot = nullptr;
+
+	if (Index >= 0 && Index < GetChildren()->Num())
+	{
+		Slot = &static_cast<SWrapBox::FSlot&>(const_cast<FSlotBase&>(GetChildren()->GetSlotAt(Index)));
+	}
+	else
+	{
+		AddSlot().Expose(Slot);
+	}
+
+	Slot->SetHorizontalAlignment(AlignmentData.HorizontalAlignment);
+	Slot->SetVerticalAlignment(AlignmentData.VerticalAlignment);
+	Slot->SetPadding(AlignmentData.Padding);
+	if (Slot->GetWidget()->GetWidgetClass().GetWidgetType() != SBox::StaticWidgetClass().GetWidgetType())
+	{
+		(*Slot)
+		[
+			SNew(SBox)
+			.MinDesiredWidth(AlignmentData.MinWidth > 0 ? AlignmentData.MinWidth : FOptionalSize())
+			.MinDesiredHeight(AlignmentData.MinHeight > 0 ? AlignmentData.MinHeight : FOptionalSize())
+			.MaxDesiredWidth(AlignmentData.MaxWidth > 0 ? AlignmentData.MaxWidth : FOptionalSize())
+			.MaxDesiredHeight(AlignmentData.MaxHeight > 0 ? AlignmentData.MaxHeight : FOptionalSize())
+			[
+				Child.GetWidgetRef()
+			]
+		];
+	}
+	else
+	{
+		TSharedRef<SBox> SlotBox = StaticCastSharedRef<SBox>(Slot->GetWidget());
+		SlotBox->SetMinDesiredWidth(AlignmentData.MinWidth > 0 ? AlignmentData.MinWidth : FOptionalSize());
+		SlotBox->SetMinDesiredHeight(AlignmentData.MinHeight > 0 ? AlignmentData.MinHeight : FOptionalSize());
+		SlotBox->SetMaxDesiredWidth(AlignmentData.MaxWidth > 0 ? AlignmentData.MaxWidth : FOptionalSize());
+		SlotBox->SetMaxDesiredHeight(AlignmentData.MaxHeight > 0 ? AlignmentData.MaxHeight : FOptionalSize());
+		SlotBox->SetContent(Child.GetWidgetRef());
+	}
+}
diff --git a/Engine/Plugins/Experimental/SlateIM/Source/SlateIM/Private/Containers/SImWrapBox.h b/Engine/Plugins/Experimental/SlateIM/Source/SlateIM/Private/Containers/SImWrapBox.h
new file mode 100644
--- /dev/null
+++ b/Engine/Plugins/Experimental/SlateIM/Source/SlateIM/Private/Containers/SImWrapBox.h
@@ -0,0 +1,26 @@
+﻿// Copyright Epic Games, Inc. All Rights Reserved.
+
+#pragma once
+
+#include "ISlateIMContainer.h"
+#include "Widgets/Layout/SWrapBox.h"
+
+
+class SImWrapBox : public SWrapBox, public ISlateIMContainer
+{
+	SLATE_DECLARE_WIDGET(SImWrapBox, SWrapBox)
+	SLATE_IM_TYPE_DATA(SImWrapBox, ISlateIMContainer)
+	
+public:
+	virtual int32 GetNumChildren() override;
+	virtual FSlateIMChild GetChild(int32 Index) override;
+
+	virtual FSlateIMChild GetContainer() override
+	{
+		return AsShared();
+	}
+
+	virtual void RemoveUnusedChildren(int32 LastUsedChildIndex) override;
+
+	virtual void UpdateChild(FSlateIMChild Child, int32 Index, const FSlateIMSlotData& AlignmentData) override;
+};
diff --git a/Engine/Plugins/Experimental/SlateIM/Source/SlateIM/Private/Implementation/SlateIM_Containers.cpp b/Engine/Plugins/Experimental/SlateIM/Source/SlateIM/Private/Implementation/SlateIM_Containers.cpp
new file mode 100644
--- /dev/null
+++ b/Engine/Plugins/Experimental/SlateIM/Source/SlateIM/Private/Implementation/SlateIM_Containers.cpp
@@ -0,0 +1,370 @@
+﻿// Copyright Epic Games, Inc. All Rights Reserved.
+
+#include "SlateIM.h"
+
+#include "Containers/SImCompoundWidget.h"
+#include "Containers/SImContextMenuAnchor.h"
+#include "Containers/SImScrollBox.h"
+#include "Containers/SImStackBox.h"
+#include "Containers/SImTableView.h"
+#include "Containers/SImWrapBox.h"
+#include "Misc/SlateIMLogging.h"
+#include "Misc/SlateIMManager.h"
+#include "Misc/SlateIMWidgetScope.h"
+
+namespace SlateIM
+{
+	void BeginStack(EOrientation Orientation)
+	{
+		TSharedPtr<SImStackBox> ContainerWidget;
+
+		// Note widget scope must fall out of scope before we push container or else we are pushing this new container into itself instead of into its parent
+		{
+			FWidgetScope<SImStackBox> Scope;
+			ContainerWidget = Scope.GetWidget();
+
+			if (!ContainerWidget)
+			{
+				ContainerWidget = SNew(SImStackBox);
+				Scope.UpdateWidget(ContainerWidget);
+			}
+
+			ContainerWidget->SetOrientation(Orientation);
+		}
+
+		FSlateIMManager::Get().PushContainer(FContainerNode(ContainerWidget));
+	}
+
+	void EndStack()
+	{
+		FSlateIMManager::Get().PopContainer<SImStackBox>();
+	}
+
+	void BeginWrap(EOrientation Orientation)
+	{
+		TSharedPtr<SImWrapBox> ContainerWidget;
+
+		// Note widget scope must fall out of scope before we push container or else we are pushing this new container into itself instead of into its parent
+		{
+			FWidgetScope<SImWrapBox> Scope;
+			ContainerWidget = Scope.GetWidget();
+
+			if (!ContainerWidget)
+			{
+				// Note, SNew cannot take a template as the argument because it converts the argument into a string for type identification. So we have to hard code this even though
+				// its essentially the same as BeginHorizonal, just with a different type
+				ContainerWidget = SNew(SImWrapBox);
+				Scope.UpdateWidget(ContainerWidget);
+			}
+
+			ContainerWidget->SetUseAllottedSize(true);
+			ContainerWidget->SetOrientation(Orientation);
+		}
+
+		FSlateIMManager::Get().PushContainer(FContainerNode(ContainerWidget));
+	}
+
+	void EndWrap()
+	{
+		FSlateIMManager::Get().PopContainer<SImWrapBox>();
+	}
+
+	void BeginHorizontalStack()
+	{
+		BeginStack(Orient_Horizontal);
+	}
+
+	void EndHorizontalStack()
+	{
+		EndStack();
+	}
+
+	void BeginVerticalStack()
+	{
+		BeginStack(Orient_Vertical);
+	}
+
+	void EndVerticalStack()
+	{
+		EndStack();
+	}
+
+	void BeginHorizontalWrap()
+	{
+		BeginWrap(Orient_Horizontal);
+	}
+
+	void EndHorizontalWrap()
+	{
+		EndWrap();
+	}
+
+	void BeginVerticalWrap()
+	{
+		BeginWrap(Orient_Vertical);
+	}
+
+	void EndVerticalWrap()
+	{
+		EndWrap();
+	}
+
+	void BeginBorder(const FSlateBrush* BackgroundImage, EOrientation Orientation, bool bAbsorbMouse, FMargin ContentPadding)
+	{
+		TSharedPtr<SImCompoundWidget> ContainerWidget;
+
+		// Note widget scope must fall out of scope before we push container or else we are pushing this new container into itself instead of into its parent
+		{
+			FWidgetScope<SImCompoundWidget> Scope(ContentPadding, HAlign_Fill, VAlign_Fill, false);
+			ContainerWidget = Scope.GetWidget();
+
+			if (!ContainerWidget)
+			{
+				ContainerWidget = SNew(SImCompoundWidget);
+					
+				Scope.UpdateWidget(ContainerWidget);
+			}
+
+
+			ContainerWidget->SetBackgroundImage(BackgroundImage);
+			ContainerWidget->SetContentPadding(ContentPadding);
+			ContainerWidget->SetAbsorbMouse(bAbsorbMouse);
+			ContainerWidget->SetOrientation(Orientation);
+		}
+
+		FSlateIMManager::Get().PushContainer(FContainerNode(ContainerWidget));
+	}
+
+	void BeginBorder(const FName BorderStyleName, EOrientation Orientation, bool bAbsorbMouse, FMargin ContentPadding)
+	{
+		BeginBorder(FAppStyle::GetBrush(BorderStyleName), Orientation, bAbsorbMouse, ContentPadding);
+	}
+
+	void EndBorder()
+	{
+		FSlateIMManager::Get().PopContainer<SImCompoundWidget>();
+	}
+
+	bool BeginScrollBox(EOrientation Orientation)
+	{
+		TSharedPtr<SImScrollBox> ContainerWidget;
+
+		bool bUserScrolled = false;
+		{
+			FWidgetScope<SImScrollBox> Scope(false);
+			ContainerWidget = Scope.GetWidget();
+
+			if (!ContainerWidget)
+			{
+				ContainerWidget =
+					SNew(SImScrollBox)
+					.Orientation(Orientation)
+					.OnUserScrolled_Lambda([ActivationData = FSlateIMManager::Get().GetWidgetActivationData()](float Offset)
+					{
+						FSlateIMManager::Get().ActivateWidget(ActivationData);
+					});
+
+				Scope.UpdateWidget(ContainerWidget);
+			}
+			else
+			{
+				bUserScrolled = FSlateIMManager::Get().IsWidgetActivatedThisFrame();
+			}
+		}
+
+		FSlateIMManager::Get().PushContainer(FContainerNode(ContainerWidget));
+
+		return bUserScrolled;
+	}
+
+	void EndScrollBox()
+	{
+		FSlateIMManager::Get().PopContainer<SImScrollBox>();
+	}
+	
+	void BeginTableRow()
+	{
+		// If we're in a row, we're adding a child row, otherwise it's a top-level row in the table
+		TSharedPtr<SImTableView> Table = FSlateIMManager::Get().GetCurrentContainer<SImTableView>();
+		TSharedPtr<FSlateIMTableRow> ParentRow = !Table.IsValid() ? FSlateIMManager::Get().GetCurrentContainer<FSlateIMTableRow>() : nullptr;
+		if (ensureMsgf(Table || ParentRow, TEXT("Table Rows/Cells can only exist within Tables")))
+		{
+			TSharedPtr<FSlateIMTableRow> TableRow;
+			// Note widget scope must fall out of scope before we push container or else we are pushing this new container into itself instead of into its parent
+			{
+				FWidgetScope<FSlateIMTableRow> Scope;
+				TableRow = Scope.GetWidget();
+
+				if (!TableRow)
+				{
+					TableRow = MakeShared<FSlateIMTableRow>();
+					Scope.UpdateWidget(TableRow);
+				}
+			}
+
+			FSlateIMManager::Get().PushContainer(FContainerNode(TableRow));
+		}
+	}
+
+	void EndTableRow()
+	{
+		FSlateIMManager::Get().PopContainer<FSlateIMTableRow>();
+	}
+	
+	void BeginTable(const FTableViewStyle* InStyle)
+	{
+		TSharedPtr<SImTableView> ContainerWidget;
+
+		// Note widget scope must fall out of scope before we push container or else we are pushing this new container into itself instead of into its parent
+		{
+			FWidgetScope<SImTableView> Scope;
+			ContainerWidget = Scope.GetWidget();
+			
+			Scope.HashData(InStyle);
+
+			if (!ContainerWidget)
+			{
+				ContainerWidget = SNew(SImTableView)
+					.TreeViewStyle(InStyle)
+					.SelectionMode(ESelectionMode::None);
+				Scope.UpdateWidget(ContainerWidget);
+			}
+			else if (Scope.IsDataHashDirty())
+			{
+				ContainerWidget->SetStyle(InStyle);
+			}
+
+			ContainerWidget->BeginTableUpdates();
+		}
+
+		FSlateIMManager::Get().PushContainer(FContainerNode(ContainerWidget));
+	}
+
+	void EndTable()
+	{
+		// If there's an open cell, close it
+		if (TSharedPtr<SImStackBox> Cell = FSlateIMManager::Get().GetCurrentContainer<SImStackBox>())
+		{
+			SlateIM::EndVerticalStack();
+		}
+
+		TSharedPtr<FSlateIMTableRow> Row = FSlateIMManager::Get().GetCurrentContainer<FSlateIMTableRow>();
+		if (ensure(Row))
+		{
+			EndTableRow();
+		}
+		
+		TSharedPtr<SImTableView> Table = FSlateIMManager::Get().GetCurrentContainer<SImTableView>();
+		if (ensure(Table))
+		{
+			FSlateIMManager::Get().PopContainer<SImTableView>();
+			Table->EndTableUpdates();
+		}
+	}
+
+	void AddTableColumn(const FStringView& ColumnLabel)
+	{
+		if (TSharedPtr<SImTableView> Table = FSlateIMManager::Get().GetCurrentContainer<SImTableView>())
+		{
+			constexpr bool bAutoSize = false;
+			const FSlateIMSlotData SlotData = FSlateIMManager::Get().GetCurrentAlignmentData(Defaults::Padding, HAlign_Fill, VAlign_Fill, bAutoSize
+				, Defaults::MinWidth, Defaults::MinHeight, Defaults::MaxWidth, Defaults::MaxHeight);
+			
+			Table->AddColumn(ColumnLabel, FSlateIMManager::Get().GetCurrentRoot().CurrentToolTip, SlotData);
+			
+			FSlateIMManager::Get().GetMutableCurrentRoot().CurrentToolTip.Empty();
+			FSlateIMManager::Get().ResetAlignmentData();
+		}
+	}
+
+	void FixedTableColumnWidth(float Width)
+	{
+		SlateIM::AutoSize();
+		SlateIM::MinWidth(Width);
+		SlateIM::MaxWidth(Width);
+	}
+
+	void InitialTableColumnWidth(float Width)
+	{
+		SlateIM::AutoSize();
+		SlateIM::MinWidth(Width);
+	}
+
+	bool NextTableCell()
+	{
+		SCOPED_NAMED_EVENT_TEXT("SlateIM::NextTableCell", FColorList::Goldenrod);
+		
+		// If there's an open cell, close it
+		if (TSharedPtr<SImStackBox> Cell = FSlateIMManager::Get().GetCurrentContainer<SImStackBox>())
+		{
+			SlateIM::EndVerticalStack();
+		}
+
+		// We haven't drawn any table cells yet, begin the first row
+		if (TSharedPtr<SImTableView> Table = FSlateIMManager::Get().GetCurrentContainer<SImTableView>())
+		{
+			Table->BeginTableContent();
+			BeginTableRow();
+		}
+
+		TSharedPtr<FSlateIMTableRow> Row = FSlateIMManager::Get().GetCurrentContainer<FSlateIMTableRow>();
+		if (ensure(Row))
+		{
+			const FContainerNode* ContainerNode = FSlateIMManager::Get().GetCurrentContainerNode();
+			check(ContainerNode);
+			
+			// Start next row if we've filled all the columns in this row
+			if (Row->GetColumnCount() == Row->CountCellWidgetsUpToIndex(ContainerNode->LastUsedChildIndex))
+			{
+				EndTableRow();
+				BeginTableRow();
+				Row = FSlateIMManager::Get().GetCurrentContainer<FSlateIMTableRow>();
+				ensure(Row);
+			}
+			
+
+			// Default table cells to fill (NextHAlign might have been set by the user, so we check first) 
+			if (!FSlateIMManager::Get().NextHAlign.IsSet())
+			{
+				FSlateIMManager::Get().NextHAlign = HAlign_Fill;
+			}
+			// Create a vertical stack to act as our cell widget
+			// TODO - Only automatically create a container when one is not created by the user
+			SlateIM::BeginVerticalStack();
+		}
+
+		return Row && Row->AreTableRowContentsRequired();
+	}
+
+	bool BeginTableRowChildren()
+	{
+		// TODO - Let caller set default expansion state
+		// TODO - Save and restore expansion state
+		
+		// If there's an open cell, close it
+		if (TSharedPtr<SImStackBox> Cell = FSlateIMManager::Get().GetCurrentContainer<SImStackBox>())
+		{
+			SlateIM::EndVerticalStack();
+		}
+
+		TSharedPtr<FSlateIMTableRow> ParentRow = FSlateIMManager::Get().GetCurrentContainer<FSlateIMTableRow>();
+		if (!ensureMsgf(ParentRow, TEXT("Child Table Rows can only be added to table rows. Did you forget to call NextTableCell()?")))
+		{
+			return false;
+		}
+
+		SlateIM::BeginTableRow();
+
+		return ParentRow->IsExpanded();
+	}
+
+	void EndTableRowChildren()
+	{
+		// If there's an open cell, close it
+		if (TSharedPtr<SImStackBox> Cell = FSlateIMManager::Get().GetCurrentContainer<SImStackBox>())
+		{
+			SlateIM::EndVerticalStack();
+		}
+		EndTableRow();
+	}
+}
diff --git a/Engine/Plugins/Experimental/SlateIM/Source/SlateIM/Private/Implementation/SlateIM_Menu.cpp b/Engine/Plugins/Experimental/SlateIM/Source/SlateIM/Private/Implementation/SlateIM_Menu.cpp
new file mode 100644
--- /dev/null
+++ b/Engine/Plugins/Experimental/SlateIM/Source/SlateIM/Private/Implementation/SlateIM_Menu.cpp
@@ -0,0 +1,97 @@
+﻿// Copyright Epic Games, Inc. All Rights Reserved.
+
+#include "SlateIM.h"
+
+#include "Containers/SImContextMenuAnchor.h"
+#include "Misc/SlateIMManager.h"
+#include "Misc/SlateIMWidgetScope.h"
+
+namespace SlateIM
+{
+	bool BeginContextMenuAnchor()
+	{
+		TSharedPtr<SImContextMenuAnchor> ContainerWidget;
+
+		bool bMenuOpened = false;
+		{
+			FWidgetScope<SImContextMenuAnchor> Scope;
+			ContainerWidget = Scope.GetWidget();
+
+			if (!ContainerWidget)
+			{
+				ContainerWidget = SNew(SImContextMenuAnchor);
+
+				Scope.UpdateWidget(ContainerWidget);
+			}
+			else
+			{
+				bMenuOpened = ContainerWidget->IsMenuOpen();
+			}
+		}
+
+
+		FSlateIMManager::Get().PushContainer(FContainerNode(ContainerWidget));
+		FSlateIMManager::Get().PushMenuRoot(ContainerWidget);
+
+		return bMenuOpened;
+	}
+
+	void EndContextMenuAnchor()
+	{
+		FSlateIMManager::Get().PopMenuRoot();
+		FSlateIMManager::Get().PopContainer<SImContextMenuAnchor>();
+	}
+
+	void AddMenuSeparator()
+	{
+		TSharedPtr<SImContextMenuAnchor> Anchor = FSlateIMManager::Get().GetCurrentMenuRoot();
+		checkf(Anchor, TEXT("Cannot add menu items without a current active menu anchor"));
+
+		return Anchor->AddMenuSeparator();
+	}
+
+	void AddMenuSection(const FStringView& SectionText)
+	{
+		TSharedPtr<SImContextMenuAnchor> Anchor = FSlateIMManager::Get().GetCurrentMenuRoot();
+		checkf(Anchor, TEXT("Cannot add menu items without a current active menu anchor"));
+
+		return Anchor->AddMenuSection(SectionText);
+	}
+
+	bool AddMenuButton(const FStringView& RowText, const FStringView& ToolTipText)
+	{
+		TSharedPtr<SImContextMenuAnchor> Anchor = FSlateIMManager::Get().GetCurrentMenuRoot();
+		checkf(Anchor, TEXT("Cannot add menu items without a current active menu anchor"));
+
+		return Anchor->AddMenuButton(RowText, ToolTipText);
+	}
+
+	bool AddMenuToggleButton(const FStringView& RowText, bool& InOutCurrentState, const FStringView& ToolTipText)
+	{
+		TSharedPtr<SImContextMenuAnchor> Anchor = FSlateIMManager::Get().GetCurrentMenuRoot();
+		checkf(Anchor, TEXT("Cannot add menu items without a current active menu anchor"));
+
+		return Anchor->AddMenuToggleButton(RowText, InOutCurrentState, ToolTipText);
+	}
+
+	bool AddMenuCheckButton(const FStringView& RowText, bool& InOutCurrentState, const FStringView& ToolTipText)
+	{
+		TSharedPtr<SImContextMenuAnchor> Anchor = FSlateIMManager::Get().GetCurrentMenuRoot();
+		checkf(Anchor, TEXT("Cannot add menu items without a current active menu anchor"));
+
+		return Anchor->AddMenuCheckButton(RowText, InOutCurrentState, ToolTipText);
+	}
+
+	void BeginSubMenu(const FStringView& SubMenuText)
+	{
+		TSharedPtr<SImContextMenuAnchor> Anchor = FSlateIMManager::Get().GetCurrentMenuRoot();
+		checkf(Anchor, TEXT("Cannot add menu items without a current active menu anchor"));
+
+		Anchor->BeginSubMenu(SubMenuText);
+	}
+
+	void EndSubMenu()
+	{
+
+	}
+}
diff --git a/Engine/Plugins/Experimental/SlateIM/Source/SlateIM/Private/Implementation/SlateIM_Roots.cpp b/Engine/Plugins/Experimental/SlateIM/Source/SlateIM/Private/Implementation/SlateIM_Roots.cpp
new file mode 100644
--- /dev/null
+++ b/Engine/Plugins/Experimental/SlateIM/Source/SlateIM/Private/Implementation/SlateIM_Roots.cpp
@@ -0,0 +1,209 @@
+﻿// Copyright Epic Games, Inc. All Rights Reserved.
+
+#include "SlateIM.h"
+
+#include "Framework/Application/SlateApplication.h"
+#include "Misc/SlateIMManager.h"
+#include "Roots/SlateIMExposedRoot.h"
+#include "Roots/SlateIMViewportRoot.h"
+#include "Roots/SlateIMWindowRoot.h"
+
+#if WITH_ENGINE
+#include "Engine/GameViewportClient.h"
+#include "Engine/LocalPlayer.h"
+#endif
+
+namespace SlateIM
+{
+	bool BeginWindowRoot(FName UniqueName, const FStringView& WindowTitle, FVector2f WindowSize, bool bShouldReopen)
+	{
+		if (!FSlateIMManager::Get().CanUpdateSlateIM())
+		{
+			return false;
+		}
+
+		bool bNeedsCreate = true;
+
+		FRootNode* RootNode = FSlateIMManager::Get().FindRoot<FSlateIMWindowRoot>(UniqueName);
+		if (RootNode)
+		{
+			// Window existed and was destroyed and needs to be reopened if requested
+			bNeedsCreate = !RootNode->RootWidget.IsValid() && bShouldReopen;
+
+			RootNode->RootState = RootNode->RootWidget->IsVisible();
+		}
+
+		if (bNeedsCreate)
+		{
+			// TODO - Save and restore window size and position, mapped to UniqueName
+			TSharedRef<SWindow> NewWindowWidget
+				= SNew(SWindow)
+				.Title(FText::FromStringView(WindowTitle))
+				.ClientSize(WindowSize);
+
+			FSlateApplication::Get().AddWindow(NewWindowWidget);
+
+			TSharedRef<FSlateIMWindowRoot> NewRoot = MakeShared<FSlateIMWindowRoot>(NewWindowWidget);
+
+			if (RootNode)
+			{
+				RootNode->RootWidget = NewRoot;
+				RootNode->RootState = true;
+			}
+			else
+			{
+				RootNode = &FSlateIMManager::Get().AddRoot(UniqueName, NewRoot);
+			}
+
+			RootNode->RootWidget->UpdateRoot(WindowTitle, FVector2f::ZeroVector, WindowSize);
+		}
+		else if (RootNode->RootState)
+		{
+			RootNode->RootWidget->UpdateRoot(WindowTitle, FVector2f::ZeroVector, WindowSize);
+		}
+
+		FSlateIMManager::Get().BeginRoot(UniqueName);
+
+		return RootNode->RootState;
+	}
+
+#if WITH_ENGINE
+	bool BeginViewportRoot(FName UniqueName, UGameViewportClient* ViewportClient)
+	{
+		if (!FSlateIMManager::Get().CanUpdateSlateIM())
+		{
+			return false;
+		}
+		
+		bool bNeedsCreate = true;
+
+		FRootNode* ViewportRoot = FSlateIMManager::Get().FindRoot<FSlateIMViewportRoot>(UniqueName);
+		if (ViewportRoot)
+		{
+			const bool bIsValidRoot = ViewportRoot->RootWidget.IsValid() && StaticCastSharedPtr<FSlateIMViewportRoot>(ViewportRoot->RootWidget)->GameViewport == ViewportClient;
+			bNeedsCreate = !bIsValidRoot;
+
+			ViewportRoot->RootState = bIsValidRoot;
+		}
+
+		if (bNeedsCreate)
+		{
+			TSharedRef<FSlateIMViewportRoot> NewRoot = MakeShared<FSlateIMViewportRoot>(ViewportClient);
+
+			NewRoot->UpdateRoot(FStringView(), FVector2f::ZeroVector, FVector2f::ZeroVector);
+			ViewportRoot = &FSlateIMManager::Get().AddRoot(UniqueName, NewRoot);
+		}
+
+		FSlateIMManager::Get().BeginRoot(UniqueName);
+
+		return ViewportRoot->RootState;
+	}
+
+	bool BeginViewportRoot(FName UniqueName, ULocalPlayer* LocalPlayer)
+	{
+		if (!FSlateIMManager::Get().CanUpdateSlateIM())
+		{
+			return false;
+		}
+		
+		bool bNeedsCreate = true;
+
+		FRootNode* ViewportRoot = FSlateIMManager::Get().FindRoot<FSlateIMViewportRoot>(UniqueName);
+		if (ViewportRoot)
+		{
+			const bool bIsValidRoot = ViewportRoot->RootWidget.IsValid() && StaticCastSharedPtr<FSlateIMViewportRoot>(ViewportRoot->RootWidget)->LocalPlayer == LocalPlayer;
+			bNeedsCreate = !bIsValidRoot;
+
+			ViewportRoot->RootState = bIsValidRoot;
+		}
+
+		if (bNeedsCreate)
+		{
+			TSharedRef<FSlateIMViewportRoot> NewRoot = MakeShared<FSlateIMViewportRoot>(LocalPlayer);
+
+			NewRoot->UpdateRoot(FStringView(), FVector2f::ZeroVector, FVector2f::ZeroVector);
+			ViewportRoot = &FSlateIMManager::Get().AddRoot(UniqueName, NewRoot);
+		}
+
+		FSlateIMManager::Get().BeginRoot(UniqueName);
+
+		return ViewportRoot->RootState;
+	}
+
+#if WITH_EDITOR
+	bool BeginViewportRoot(FName UniqueName, TSharedPtr<IAssetViewport> AssetViewport, FVector2f Offset, FVector2f NormalizedSize)
+	{
+		if (!FSlateIMManager::Get().CanUpdateSlateIM())
+		{
+			return false;
+		}
+		
+		bool bNeedsCreate = true;
+
+		FRootNode* ViewportRoot = FSlateIMManager::Get().FindRoot<FSlateIMViewportRoot>(UniqueName);
+		if (ViewportRoot)
+		{
+			const bool bIsValidRoot = ViewportRoot->RootWidget.IsValid() && StaticCastSharedPtr<FSlateIMViewportRoot>(ViewportRoot->RootWidget)->AssetViewport == AssetViewport;
+			bNeedsCreate = !bIsValidRoot;
+
+			ViewportRoot->RootState = bIsValidRoot;
+		}
+
+		if (bNeedsCreate)
+		{
+			TSharedRef<FSlateIMViewportRoot> NewRoot = MakeShared<FSlateIMViewportRoot>(AssetViewport);
+
+			NewRoot->UpdateRoot(FStringView(), Offset, NormalizedSize);
+			ViewportRoot = &FSlateIMManager::Get().AddRoot(UniqueName, NewRoot);
+		}
+		else if (ViewportRoot->RootState)
+		{
+			ViewportRoot->RootWidget->UpdateRoot(FStringView(), Offset, NormalizedSize);
+		}
+
+		FSlateIMManager::Get().BeginRoot(UniqueName);
+
+		return ViewportRoot->RootState;
+	}
+#endif // WITH_EDITOR
+
+#endif // WITH_ENGINE
+
+	bool BeginExposedRoot(FName UniqueName, TSharedPtr<SWidget>& OutSlateIMWidget)
+	{
+		if (!FSlateIMManager::Get().CanUpdateSlateIM())
+		{
+			return false;
+		}
+		
+		bool bNeedsCreate = true;
+		
+		FRootNode* ExistingRoot = FSlateIMManager::Get().FindRoot<FSlateIMExposedRoot>(UniqueName);
+		if (ExistingRoot)
+		{
+			ExistingRoot->RootState = ExistingRoot->RootWidget.IsValid();
+			bNeedsCreate = !ExistingRoot->RootWidget.IsValid();
+		}
+		
+		if (bNeedsCreate)
+		{
+			TSharedRef<FSlateIMExposedRoot> NewRoot = MakeShared<FSlateIMExposedRoot>();
+			NewRoot->UpdateRoot(FStringView(), FVector2f::ZeroVector, FVector2f::UnitVector);
+			ExistingRoot = &FSlateIMManager::Get().AddRoot(UniqueName, NewRoot);
+			OutSlateIMWidget = NewRoot->GetExposedWidget();
+		}
+		else if (TSharedPtr<FSlateIMExposedRoot> ExposedRoot = StaticCastSharedPtr<FSlateIMExposedRoot>(ExistingRoot->RootWidget))
+		{
+			OutSlateIMWidget = ExposedRoot->GetExposedWidget();
+		}
+
+		FSlateIMManager::Get().BeginRoot(UniqueName);
+		
+		return ExistingRoot->RootState;
+	}
+
+	void EndRoot()
+	{
+		FSlateIMManager::Get().EndRoot();
+	}
+}
diff --git a/Engine/Plugins/Experimental/SlateIM/Source/SlateIM/Private/Implementation/SlateIM_Slots.cpp b/Engine/Plugins/Experimental/SlateIM/Source/SlateIM/Private/Implementation/SlateIM_Slots.cpp
new file mode 100644
--- /dev/null
+++ b/Engine/Plugins/Experimental/SlateIM/Source/SlateIM/Private/Implementation/SlateIM_Slots.cpp
@@ -0,0 +1,53 @@
+﻿// Copyright Epic Games, Inc. All Rights Reserved.
+
+#include "SlateIM.h"
+
+#include "Misc/SlateIMManager.h"
+
+namespace SlateIM
+{
+	void Padding(const FMargin NextPadding)
+	{
+		FSlateIMManager::Get().NextPadding = NextPadding;
+	}
+
+	void HAlign(EHorizontalAlignment NextAlignment)
+	{
+		FSlateIMManager::Get().NextHAlign = NextAlignment;
+	}
+	
+	void VAlign(EVerticalAlignment NextAlignment)
+	{
+		FSlateIMManager::Get().NextVAlign = NextAlignment;
+	}
+
+	void AutoSize()
+	{
+		FSlateIMManager::Get().NextAutoSize = true;
+	}
+
+	void Fill()
+	{
+		FSlateIMManager::Get().NextAutoSize = false;
+	}
+
+	void MinWidth(float InMinWidth)
+	{
+		FSlateIMManager::Get().NextMinWidth = InMinWidth;
+	}
+
+	void MinHeight(float InMinHeight)
+	{
+		FSlateIMManager::Get().NextMinHeight = InMinHeight;
+	}
+
+	void MaxWidth(float InMaxWidth)
+	{
+		FSlateIMManager::Get().NextMaxWidth = InMaxWidth;
+	}
+
+	void MaxHeight(float InMaxHeight)
+	{
+		FSlateIMManager::Get().NextMaxHeight = InMaxHeight;
+	}
+}
diff --git a/Engine/Plugins/Experimental/SlateIM/Source/SlateIM/Private/Implementation/SlateIM_Util.cpp b/Engine/Plugins/Experimental/SlateIM/Source/SlateIM/Private/Implementation/SlateIM_Util.cpp
new file mode 100644
--- /dev/null
+++ b/Engine/Plugins/Experimental/SlateIM/Source/SlateIM/Private/Implementation/SlateIM_Util.cpp
@@ -0,0 +1,40 @@
+﻿// Copyright Epic Games, Inc. All Rights Reserved.
+
+#include "Misc/MessageDialog.h"
+#include "SlateIM.h"
+
+#include "Misc/SlateIMManager.h"
+
+namespace SlateIM
+{
+	bool CanUpdateSlateIM()
+	{
+		return FSlateIMManager::Get().CanUpdateSlateIM();
+	}
+	
+	void BeginDisabledState()
+	{
+		FSlateIMManager::Get().GetMutableCurrentRoot().SetDisabledState();
+	}
+
+	void EndDisabledState()
+	{
+		FSlateIMManager::Get().GetMutableCurrentRoot().SetEnabledState();
+	}
+
+	void SetToolTip(const FStringView& NextToolTip)
+	{
+		FSlateIMManager::Get().GetMutableCurrentRoot().SetNextToolTip(NextToolTip);
+	}
+
+	EAppReturnType::Type ModalDialog(EAppMsgType::Type MessageType, const FStringView& DialogText, EAppMsgCategory Category, const FStringView& DialogTitle)
+	{
+		FSlateIMManager::Get().OnSlateIMModalOpened();
+
+		ON_SCOPE_EXIT
+		{
+			FSlateIMManager::Get().OnSlateIMModalClosed();
+		};
+		return FMessageDialog::Open(Category, MessageType, FText::FromStringView(DialogText), FText::FromStringView(DialogTitle));
+	}
+}
diff --git a/Engine/Plugins/Experimental/SlateIM/Source/SlateIM/Private/Implementation/SlateIM_Widgets.cpp b/Engine/Plugins/Experimental/SlateIM/Source/SlateIM/Private/Implementation/SlateIM_Widgets.cpp
new file mode 100644
--- /dev/null
+++ b/Engine/Plugins/Experimental/SlateIM/Source/SlateIM/Private/Implementation/SlateIM_Widgets.cpp
@@ -0,0 +1,719 @@
+﻿// Copyright Epic Games, Inc. All Rights Reserved.
+
+#include "SlateIM.h"
+
+#include "Misc/SlateIMLogging.h"
+#include "Misc/SlateIMManager.h"
+#include "Misc/SlateIMSlotData.h"
+#include "Misc/SlateIMWidgetScope.h"
+#include "Widgets/Images/SImage.h"
+#include "Widgets/Input/SEditableTextBox.h"
+#include "Widgets/Input/SSlider.h"
+#include "Widgets/Input/SSpinBox.h"
+#include "Widgets/Input/STextComboBox.h"
+#include "Widgets/Layout/SSpacer.h"
+#include "Widgets/Notifications/SProgressBar.h"
+#include "Widgets/SImButton.h"
+#include "Widgets/SImCheckBox.h"
+#include "Widgets/Text/STextBlock.h"
+
+#if WITH_ENGINE
+#include "Engine/Texture2D.h"
+#include "Engine/TextureRenderTarget2D.h"
+#include "SlateMaterialBrush.h"
+#endif
+
+namespace SlateIM
+{
+	template<typename DataType>
+	struct FSlateIMDataStore : ISlateMetaData
+	{
+		SLATE_METADATA_TYPE(FSlateIMDataStore, ISlateMetaData);
+
+		FSlateIMDataStore(DataType&& InData)
+			: Data(MoveTemp(InData))
+		{}
+
+		FSlateIMDataStore(const DataType& InData)
+			: Data(InData)
+		{}
+
+		DataType Data;
+	};
+	
+	void Text(const FStringView& InText, const FTextBlockStyle* TextStyle /*= nullptr*/)
+	{
+		Text(InText, FSlateColor::UseForeground(), TextStyle);
+	}
+
+	void Text(const FStringView& InText, FSlateColor Color, const FTextBlockStyle* TextStyle /*= nullptr*/)
+	{
+		FWidgetScope<STextBlock> Scope;
+		TSharedPtr<STextBlock> TextBlock = Scope.GetWidget();
+
+		Scope.HashData(Color);
+		Scope.HashData(TextStyle);
+		Scope.HashStringView(InText);
+
+		if (!TextBlock)
+		{
+			TextBlock =
+				SNew(STextBlock)
+				.Text(FText::FromStringView(InText))
+				.ColorAndOpacity(Color);
+
+			if (TextStyle)
+			{
+				TextBlock->SetTextStyle(TextStyle);
+			}
+
+			Scope.UpdateWidget(TextBlock);
+		}
+		else if(Scope.IsDataHashDirty())
+		{
+			TextBlock->SetText(FText::FromStringView(InText));
+			TextBlock->SetTextStyle(TextStyle);
+			TextBlock->SetColorAndOpacity(Color);
+		}
+	}
+
+	bool EditableText(FString& InOutText, const FStringView& HintText, const FEditableTextBoxStyle* TextStyle)
+	{
+		FWidgetScope<SEditableTextBox> Scope;
+		TSharedPtr<SEditableTextBox> EditableText = Scope.GetWidget();
+		const float MinWidth = FSlateIMManager::Get().NextMinWidth.Get(Defaults::InputWidgetWidth);
+
+		bool bWasActivated = false;
+
+		Scope.HashData(TextStyle);
+		Scope.HashStringView(InOutText);
+		Scope.HashStringView(HintText);
+
+		if (!EditableText)
+		{
+			FWidgetActivation ActivationData = FSlateIMManager::Get().GetWidgetActivationData();
+
+			EditableText =
+				SNew(SEditableTextBox)
+				.MinDesiredWidth(MinWidth)
+				.Text(FText::FromStringView(InOutText))
+				.HintText(FText::FromStringView(HintText))
+				.OnTextChanged_Lambda([ActivationData](const FText& NewText)
+				{
+					FSlateIMManager::Get().ActivateWidget(ActivationData);
+				})
+				.OnTextCommitted_Lambda([ActivationData](const FText& NewText, ETextCommit::Type)
+				{
+					FSlateIMManager::Get().ActivateWidget(ActivationData);
+				});
+
+			if (TextStyle)
+			{
+				EditableText->SetStyle(TextStyle);
+			}
+			
+			Scope.UpdateWidget(EditableText);
+		}
+		else
+		{
+			bWasActivated = FSlateIMManager::Get().IsWidgetActivatedThisFrame();
+
+			const bool bIsHashDirty = Scope.IsDataHashDirty();
+
+			if (bIsHashDirty)
+			{
+				EditableText->SetStyle(TextStyle);
+				EditableText->SetHintText(FText::FromStringView(HintText));
+			}
+			
+			if (bWasActivated)
+			{
+				InOutText = EditableText->GetText().ToString();
+			}
+			else if (bIsHashDirty)
+			{
+				EditableText->SetText(FText::FromStringView(InOutText));
+			}
+			
+			EditableText->SetMinimumDesiredWidth(MinWidth);
+		}
+
+		return bWasActivated;
+	}
+
+	void Image_Internal(const FSlateColor& ColorAndOpacity, FVector2D DesiredSize, TFunctionRef<const FSlateBrush*(const TSharedRef<SImage>&)> GetBrush)
+	{
+		FWidgetScope<SImage> Scope(Defaults::Padding, HAlign_Left, VAlign_Top);
+		TSharedPtr<SImage> ImageWidget = Scope.GetWidget();
+
+		if (!ImageWidget)
+		{
+			ImageWidget =
+				SNew(SImage)
+				.ColorAndOpacity(ColorAndOpacity)
+				.DesiredSizeOverride(DesiredSize == FVector2D::ZeroVector ? TOptional<FVector2D>() : DesiredSize);
+
+			ImageWidget->SetImage(GetBrush(ImageWidget.ToSharedRef()));
+			Scope.UpdateWidget(ImageWidget);
+		}
+		else
+		{
+			ImageWidget->SetImage(GetBrush(ImageWidget.ToSharedRef()));
+			ImageWidget->SetColorAndOpacity(ColorAndOpacity);
+			ImageWidget->SetDesiredSizeOverride(DesiredSize == FVector2D::ZeroVector ? TOptional<FVector2D>() : DesiredSize);
+		}
+	}
+
+	void Image(const FSlateBrush* ImageBrush, const FSlateColor& ColorAndOpacity, FVector2D DesiredSize)
+	{
+		Image_Internal(ColorAndOpacity, DesiredSize, [ImageBrush](const TSharedRef<SImage>&) { return ImageBrush; });
+	}
+
+	void Image(const FName ImageStyleName, const FSlateColor& ColorAndOpacity, FVector2D DesiredSize)
+	{
+		Image(FAppStyle::Get().GetBrush(ImageStyleName), ColorAndOpacity, DesiredSize);
+	}
+
+#if WITH_ENGINE
+	struct FPinnedImageResource
+	{
+		FSlateBrush Brush;
+		TStrongObjectPtr<UObject> PinnedResource;
+	};
+	using FUObjectImageResource = FSlateIMDataStore<FPinnedImageResource>;
+	void Image(UTexture2D* ImageTexture, const FSlateColor& ColorAndOpacity, FVector2D DesiredSize)
+	{
+		auto MakeTextureBrush = [ImageTexture]() -> FSlateBrush
+		{
+			if (ImageTexture)
+			{
+				FSlateBrush Brush;
+				Brush.SetResourceObject(ImageTexture);
+				Brush.ImageSize = FVector2D(ImageTexture->GetSizeX(), ImageTexture->GetSizeY());
+				return Brush;
+			}
+			return FSlateNoResource();
+		};
+		
+		Image_Internal(ColorAndOpacity, DesiredSize, [&MakeTextureBrush, ImageTexture](const TSharedRef<SImage>& ImageWidget)
+		{
+			TSharedPtr<FUObjectImageResource> Resource = ImageWidget->GetMetaData<FUObjectImageResource>();
+			if (!Resource)
+			{
+				Resource = MakeShared<FUObjectImageResource>(FPinnedImageResource{ MakeTextureBrush(), TStrongObjectPtr<UObject>(ImageTexture) });
+				ImageWidget->AddMetadata(Resource.ToSharedRef());
+			}
+			else if (Resource->Data.PinnedResource.Get() != ImageTexture)
+			{
+				Resource->Data = FPinnedImageResource{ MakeTextureBrush(), TStrongObjectPtr<UObject>(ImageTexture) };
+			}
+
+			return &Resource->Data.Brush;
+		});
+	}
+
+	void Image(UTextureRenderTarget2D* ImageRenderTarget, const FSlateColor& ColorAndOpacity, FVector2D DesiredSize)
+	{
+		auto MakeRTBrush = [ImageRenderTarget]() -> FSlateBrush
+		{
+			if (ImageRenderTarget)
+			{
+				FSlateBrush Brush;
+				Brush.SetResourceObject(ImageRenderTarget);
+				Brush.ImageSize = FVector2D(ImageRenderTarget->SizeX, ImageRenderTarget->SizeY);
+				return Brush;
+			}
+			return FSlateNoResource();
+		};
+		
+		Image_Internal(ColorAndOpacity, DesiredSize, [&MakeRTBrush, ImageRenderTarget](const TSharedRef<SImage>& ImageWidget)
+		{
+			TSharedPtr<FUObjectImageResource> Resource = ImageWidget->GetMetaData<FUObjectImageResource>();
+			if (!Resource)
+			{
+				Resource = MakeShared<FUObjectImageResource>(FPinnedImageResource{ MakeRTBrush(), TStrongObjectPtr<UObject>(ImageRenderTarget) });
+				ImageWidget->AddMetadata(Resource.ToSharedRef());
+			}
+			else if (Resource->Data.PinnedResource.Get() != ImageRenderTarget)
+			{
+				Resource->Data = FPinnedImageResource{ MakeRTBrush(), TStrongObjectPtr<UObject>(ImageRenderTarget) };
+			}
+
+			return &Resource->Data.Brush;
+		});
+	}
+
+	void Image(UMaterialInterface* ImageMaterial, FVector2D BrushSize, const FSlateColor& ColorAndOpacity, FVector2D DesiredSize)
+	{
+		auto MakeMaterialBrush = [ImageMaterial, &BrushSize]() -> FSlateBrush
+		{
+			if (ImageMaterial)
+			{
+				return FSlateMaterialBrush(*ImageMaterial, BrushSize);
+			}
+			return FSlateNoResource();
+		};
+		
+		Image_Internal(ColorAndOpacity, DesiredSize, [&MakeMaterialBrush, ImageMaterial](const TSharedRef<SImage>& ImageWidget)
+		{
+			TSharedPtr<FUObjectImageResource> Resource = ImageWidget->GetMetaData<FUObjectImageResource>();
+			if (!Resource)
+			{
+				Resource = MakeShared<FUObjectImageResource>(FPinnedImageResource{ MakeMaterialBrush(), TStrongObjectPtr<UObject>(ImageMaterial) });
+				ImageWidget->AddMetadata(Resource.ToSharedRef());
+			}
+			else if (Resource->Data.PinnedResource.Get() != ImageMaterial)
+			{
+				Resource->Data = FPinnedImageResource{ MakeMaterialBrush(), TStrongObjectPtr<UObject>(ImageMaterial) };
+			}
+
+			return &Resource->Data.Brush;
+		});
+	}
+#endif
+
+	bool Button(const FStringView& InText, const FButtonStyle* InStyle)
+	{
+		FWidgetScope<SImButton> Scope;
+		TSharedPtr<SImButton> ButtonWidget = Scope.GetWidget();
+
+		Scope.HashStringView(InText);
+
+		bool bWasClicked = false;
+		if (!ButtonWidget)
+		{
+			ButtonWidget =
+				SNew(SImButton)
+				.OnClicked_Lambda([ActivationData = FSlateIMManager::Get().GetWidgetActivationData()]()
+				{
+					FSlateIMManager::Get().ActivateWidget(ActivationData);
+					return FReply::Handled();
+				});
+
+			ButtonWidget->SetText(InText);
+
+			if (InStyle)
+			{
+				ButtonWidget->SetButtonStyle(InStyle);
+			}
+
+			Scope.UpdateWidget(ButtonWidget);
+		}
+		else
+		{
+			bWasClicked = FSlateIMManager::Get().IsWidgetActivatedThisFrame();
+
+			if (Scope.IsDataHashDirty())
+			{
+				ButtonWidget->SetText(InText);
+			}
+		}
+
+		return bWasClicked;
+	}
+
+	bool CheckBox(const FStringView& InText, bool& InOutCurrentState, const FCheckBoxStyle* CheckBoxStyle)
+	{
+		ECheckBoxState CurrentEnumState = InOutCurrentState ? ECheckBoxState::Checked : ECheckBoxState::Unchecked;
+
+		const bool bValueChanged = CheckBox(InText, CurrentEnumState, CheckBoxStyle);
+		if (bValueChanged)
+		{
+			InOutCurrentState = CurrentEnumState == ECheckBoxState::Checked;
+		}
+
+		return bValueChanged;
+	}
+
+	bool CheckBox(const FStringView& InText, ECheckBoxState& InOutCurrentState, const FCheckBoxStyle* CheckBoxStyle)
+	{
+		FWidgetScope<SImCheckBox> Scope;
+		TSharedPtr<SImCheckBox> CheckboxWidget = Scope.GetWidget();
+
+		bool bValueChanged = false;
+
+		Scope.HashStringView(InText);
+
+		if (!CheckboxWidget)
+		{
+			CheckboxWidget =
+				SNew(SImCheckBox)
+				.IsChecked(InOutCurrentState)
+				.OnCheckStateChanged_Lambda([ActivationData = FSlateIMManager::Get().GetWidgetActivationData()](ECheckBoxState NewState)
+				{
+					FSlateIMManager::Get().ActivateWidget(ActivationData);
+				});
+
+			CheckboxWidget->SetText(FText::FromStringView(InText));
+
+			if (CheckBoxStyle)
+			{
+				CheckboxWidget->SetStyle(CheckBoxStyle);
+			}
+
+			Scope.UpdateWidget(CheckboxWidget);
+		}
+		else
+		{
+			bValueChanged = FSlateIMManager::Get().IsWidgetActivatedThisFrame();
+			if (bValueChanged)
+			{
+				InOutCurrentState = CheckboxWidget->GetCheckedState();
+			}
+			else
+			{
+				CheckboxWidget->SetIsChecked(InOutCurrentState);
+			}
+
+			if (Scope.IsDataHashDirty())
+			{
+				CheckboxWidget->SetText(FText::FromStringView(InText));
+			}
+		}
+
+		return bValueChanged;
+	}
+
+	template<typename NumericType>
+	bool SpinBox_Internal(NumericType& InOutValue, TOptional<NumericType> Min, TOptional<NumericType> Max, const FSpinBoxStyle* SpinBoxStyle = nullptr)
+	{
+		FWidgetScope<SSpinBox<NumericType>> Scope(Defaults::Padding, Defaults::HAlign, Defaults::VAlign, Defaults::bAutoSize, Defaults::InputWidgetWidth);
+		TSharedPtr<SSpinBox<NumericType>> SpinBoxWidget = Scope.GetWidget();
+
+		Scope.HashData(SpinBoxStyle);
+
+		class FSpinBoxState : public ISlateMetaData
+		{
+		public:
+			SLATE_METADATA_TYPE(FSpinBoxState, ISlateMetaData)
+			bool bIsChanging = false;
+		};
+
+		bool bValueChanged = false;
+		if (!SpinBoxWidget)
+		{
+			TSharedRef<FSpinBoxState> SpinBoxState = MakeShared<FSpinBoxState>();
+			FWidgetActivation ActivationData = FSlateIMManager::Get().GetWidgetActivationData();
+			SpinBoxWidget =
+				SNew(SSpinBox<NumericType>)
+				.MinValue(Min)
+				.MaxValue(Max)
+				.Value(InOutValue)
+				.OnBeginSliderMovement_Lambda([SpinBoxState](){ SpinBoxState->bIsChanging = true; })
+				.OnEndSliderMovement_Lambda([SpinBoxState, ActivationData](NumericType NewVal)
+				{
+					SpinBoxState->bIsChanging = false; FSlateIMManager::Get().ActivateWidget(ActivationData);
+				})
+				.OnValueChanged_Lambda([ActivationData](NumericType NewVal)
+				{
+					FSlateIMManager::Get().ActivateWidget(ActivationData);
+				})
+				.OnValueCommitted_Lambda([ActivationData](NumericType NewVal, ETextCommit::Type CommitType)
+				{
+					FSlateIMManager::Get().ActivateWidget(ActivationData);
+				});
+
+			SpinBoxWidget->AddMetadata(SpinBoxState);
+			Scope.UpdateWidget(SpinBoxWidget);
+
+			if (SpinBoxStyle)
+			{
+				SpinBoxWidget->SetWidgetStyle(SpinBoxStyle);
+			}
+		}
+		else
+		{
+			SpinBoxWidget->SetMinValue(Min);
+			SpinBoxWidget->SetMaxValue(Max);
+
+			bValueChanged = FSlateIMManager::Get().IsWidgetActivatedThisFrame();
+			if (bValueChanged)
+			{
+				InOutValue = SpinBoxWidget->GetValue();
+			}
+			else if (SpinBoxWidget->GetValue() != InOutValue)
+			{
+				TSharedPtr<FSpinBoxState> SpinBoxState = SpinBoxWidget->template GetMetaData<FSpinBoxState>();
+				if (ensure(SpinBoxState) && !SpinBoxState->bIsChanging)
+				{
+					SpinBoxWidget->SetValue(InOutValue);
+				}
+			}
+
+			if (Scope.IsDataHashDirty())
+			{
+				SpinBoxWidget->SetWidgetStyle(SpinBoxStyle);
+				SpinBoxWidget->InvalidateStyle();
+			}
+		}
+
+		return bValueChanged;
+	}
+
+
+	bool SpinBox(float& InOutValue, TOptional<float> Min, TOptional<float> Max, const FSpinBoxStyle* SpinBoxStyle)
+	{
+		return SpinBox_Internal<float>(InOutValue, Min, Max, SpinBoxStyle);
+	}
+
+	bool SpinBox(double& InOutValue, TOptional<double> Min, TOptional<double> Max, const FSpinBoxStyle* SpinBoxStyle)
+	{
+		return SpinBox_Internal<double>(InOutValue, Min, Max, SpinBoxStyle);
+	}
+
+	bool SpinBox(int32& InOutValue, TOptional<int32> Min, TOptional<int32> Max, const FSpinBoxStyle* SpinBoxStyle)
+	{
+		return SpinBox_Internal<int32>(InOutValue, Min, Max, SpinBoxStyle);
+	}
+
+	bool Slider(float& InOutValue, float Min, float Max, float Step, const FSliderStyle* SliderStyle /*= nullptr*/)
+	{
+		FWidgetScope<SSlider> Scope(Defaults::Padding, Defaults::HAlign, Defaults::VAlign, Defaults::bAutoSize, Defaults::InputWidgetWidth);
+		TSharedPtr<SSlider> SliderWidget = Scope.GetWidget();
+
+		bool bValueChanged = false;
+		if (!SliderWidget)
+		{
+			SliderWidget =
+				SNew(SSlider)
+				.MinValue(Min)
+				.MaxValue(Max)
+				.StepSize(Step)
+				.Value(InOutValue)
+				.OnValueChanged_Lambda([ActivationData = FSlateIMManager::Get().GetWidgetActivationData()](float NewVal)
+				{
+					FSlateIMManager::Get().ActivateWidget(ActivationData);
+				});
+
+			if (SliderStyle)
+			{
+				SliderWidget->SetStyle(SliderStyle);
+			}
+
+			Scope.UpdateWidget(SliderWidget);
+		}
+		else
+		{
+			SliderWidget->SetMinAndMaxValues(Min, Max);
+			SliderWidget->SetStepSize(Step);
+
+			bValueChanged = FSlateIMManager::Get().IsWidgetActivatedThisFrame();
+			if (bValueChanged)
+			{
+				InOutValue = SliderWidget->GetValue();
+			}
+			else
+			{
+				SliderWidget->SetValue(InOutValue);
+			}
+		}
+
+		return bValueChanged;
+	}
+
+	void ProgressBar(TOptional<float> Percent, const FProgressBarStyle* ProgressBarStyle /*= nullptr*/)
+	{
+		FWidgetScope<SProgressBar> Scope(Defaults::Padding, Defaults::HAlign, Defaults::VAlign, Defaults::bAutoSize, Defaults::InputWidgetWidth);
+		TSharedPtr<SProgressBar> ProgressBarWidget = Scope.GetWidget();
+
+		if (!ProgressBarWidget)
+		{
+			ProgressBarWidget =
+				SNew(SProgressBar)
+				.Percent(Percent);
+
+			if (ProgressBarStyle)
+			{
+				ProgressBarWidget->SetStyle(ProgressBarStyle);
+			}
+
+			Scope.UpdateWidget(ProgressBarWidget);
+		}
+		else
+		{
+			ProgressBarWidget->SetPercent(Percent);
+		}
+	}
+
+	bool ComboBox(const TArray<FString>& ComboItems, int32& InOutSelectedItemIndex, const FComboBoxStyle* InComboStyle)
+	{
+		FWidgetScope<STextComboBox> Scope;
+		TSharedPtr<STextComboBox> ComboWidget = Scope.GetWidget();
+
+		Scope.HashData(InComboStyle);
+		
+		using FComboBoxData = FSlateIMDataStore<TArray<TSharedPtr<FString>>>;
+
+		bool bValueChanged = false;
+		if (!ComboWidget)
+		{
+			TArray<TSharedPtr<FString>> Options;
+			Options.Reserve(ComboItems.Num());
+
+			for(const FString& Item : ComboItems)
+			{
+				Options.Add(MakeShared<FString>(Item));
+			}
+		
+			TSharedRef<FComboBoxData> WidgetIMData = MakeShared<FComboBoxData>(MoveTemp(Options));
+
+			ComboWidget =
+				SNew(STextComboBox)
+				.OptionsSource(&WidgetIMData->Data)
+				.OnSelectionChanged_Lambda([ActivationData = FSlateIMManager::Get().GetWidgetActivationData()](TSharedPtr<FString> NewVal, ESelectInfo::Type SelectType)
+				{
+					FSlateIMManager::Get().ActivateWidget(ActivationData);
+				});
+
+			if (InComboStyle)
+			{
+				ComboWidget->SetStyle(InComboStyle);
+			}
+			
+			ComboWidget->AddMetadata(WidgetIMData);
+
+			Scope.UpdateWidget(ComboWidget);
+		}
+		else
+		{
+			TSharedPtr<FComboBoxData> ComboBoxData = ComboWidget->GetMetaData<FComboBoxData>();
+			bValueChanged = FSlateIMManager::Get().IsWidgetActivatedThisFrame();
+			if (bValueChanged)
+			{
+				InOutSelectedItemIndex = ComboBoxData->Data.IndexOfByKey(ComboWidget->GetSelectedItem());
+			}
+			else
+			{
+				if (InOutSelectedItemIndex != INDEX_NONE && !ComboWidget->IsOpen())
+				{
+					ComboWidget->SetSelectedItem(ComboBoxData->Data[InOutSelectedItemIndex]);
+				}
+			}
+
+			if (Scope.IsDataHashDirty())
+			{
+				ComboWidget->SetStyle(InComboStyle);
+			}
+		}
+
+		return bValueChanged;
+	}
+
+	bool SelectionList(const TArray<FString>& ListItems, int32& InOutSelectedItemIndex, bool bForceRefresh, const FTableViewStyle* InStyle)
+	{
+		SCOPED_NAMED_EVENT_TEXT("SlateIM::SelectionList", FColorList::Goldenrod);
+		using ListViewType = SListView<TSharedPtr<FString>>;
+		using FListViewData = FSlateIMDataStore<TArray<TSharedPtr<FString>>>;
+
+		FWidgetScope<ListViewType> Scope;
+		TSharedPtr<ListViewType> ListWidget = Scope.GetWidget();
+
+		Scope.HashData(InStyle);
+
+		auto BuildListItems = [&]()
+		{
+			TArray<TSharedPtr<FString>> ListItemsSource;
+			ListItemsSource.Reserve(ListItems.Num());
+			for(const FString& Item : ListItems)
+			{
+				ListItemsSource.Add(MakeShared<FString>(Item));
+			}
+			return ListItemsSource;
+		};
+
+		bool bSelectionChanged = false;
+		if (!ListWidget)
+		{
+			TSharedRef<FListViewData> ListViewData = MakeShared<FListViewData>(BuildListItems());
+			ListWidget = SNew(ListViewType)
+				.SelectionMode(ESelectionMode::Single)
+				.ListViewStyle(InStyle)
+				.ListItemsSource(&(ListViewData->Data))
+				.OnSelectionChanged_Lambda([ActivationData = FSlateIMManager::Get().GetWidgetActivationData()](TSharedPtr<FString> NewVal, ESelectInfo::Type SelectType)
+				{
+					UE_LOG(LogSlateIM, Verbose, TEXT("Selected %s"), NewVal.IsValid() ? **NewVal.Get() : TEXT("[NULL]"));
+					FSlateIMManager::Get().ActivateWidget(ActivationData);
+				})
+				.OnGenerateRow_Lambda([](TSharedPtr<FString> ListItem, const TSharedRef<STableViewBase>& OwnerTable)
+				{
+					return SNew(STableRow<TSharedPtr<FString>>, OwnerTable)
+						.Padding(4.f)
+						[
+							SNew(STextBlock).Text(FText::FromString(*ListItem))
+						];
+				});
+
+			ListWidget->AddMetadata<FListViewData>(ListViewData);
+			
+			if (ListViewData->Data.IsValidIndex(InOutSelectedItemIndex))
+			{
+				ListWidget->SetItemSelection(ListViewData->Data[InOutSelectedItemIndex], true);
+			}
+			Scope.UpdateWidget(ListWidget);
+		}
+		else
+		{
+			TSharedPtr<FListViewData> ListViewData = ListWidget->GetMetaData<FListViewData>();
+			if (bForceRefresh)
+			{
+				ListViewData->Data = BuildListItems();
+				ListWidget->RequestListRefresh();
+			
+				if (ListViewData->Data.IsValidIndex(InOutSelectedItemIndex))
+				{
+					ListWidget->SetItemSelection(ListViewData->Data[InOutSelectedItemIndex], true);
+				}
+				else
+				{
+					ListWidget->ClearSelection();
+				}
+			}
+			else
+			{
+				bSelectionChanged = FSlateIMManager::Get().IsWidgetActivatedThisFrame();
+				if (bSelectionChanged)
+				{
+					TArray<TSharedPtr<FString>> SelectedItems = ListWidget->GetSelectedItems();
+					const int32 NewSelectedIndex = SelectedItems.Num() > 0 ? ListViewData->Data.IndexOfByKey(SelectedItems[0]) : INDEX_NONE;
+					UE_LOG(LogSlateIM, Verbose, TEXT("Selection Changed %d -> %d"), InOutSelectedItemIndex, NewSelectedIndex);
+					InOutSelectedItemIndex = NewSelectedIndex;
+				}
+			}
+
+			if (Scope.IsDataHashDirty())
+			{
+				ListWidget->SetStyle(InStyle);
+			}
+		}
+
+		return bSelectionChanged;
+	}
+
+	void Spacer(const FVector2D& Size)
+	{
+		FWidgetScope<SSpacer> Scope(Defaults::Padding, HAlign_Left, VAlign_Top);
+		TSharedPtr<SSpacer> SpacerWidget = Scope.GetWidget();
+
+		if (!SpacerWidget)
+		{
+			SpacerWidget = SNew(SSpacer)
+				.Size(Size);
+			Scope.UpdateWidget(SpacerWidget);
+		}
+		else
+		{
+			SpacerWidget->SetSize(Size);
+		}
+	}
+
+	void Widget(TSharedRef<SWidget> InWidget)
+	{
+		FWidgetScope<SWidget> Scope(InWidget, Defaults::Padding, HAlign_Left, VAlign_Top);
+		TSharedPtr<SWidget> Widget = Scope.GetWidget();
+
+		if (!Widget)
+		{
+			Scope.UpdateWidget(InWidget);
+		}
+	}
+}
diff --git a/Engine/Plugins/Experimental/SlateIM/Source/SlateIM/Private/Misc/ISlateIMChild.h b/Engine/Plugins/Experimental/SlateIM/Source/SlateIM/Private/Misc/ISlateIMChild.h
new file mode 100644
--- /dev/null
+++ b/Engine/Plugins/Experimental/SlateIM/Source/SlateIM/Private/Misc/ISlateIMChild.h
@@ -0,0 +1,76 @@
+﻿// Copyright Epic Games, Inc. All Rights Reserved.
+
+#pragma once
+
+#include "Containers/Union.h"
+#include "SlateIMTypeChecking.h"
+#include "Templates/SharedPointer.h"
+#include "Widgets/SNullWidget.h"
+#include "Widgets/SWidget.h"
+
+class SWidget;
+
+// Interface for non-widget children in the SlateIM hierarchy (eg. for virtualized listview entries)
+class ISlateIMChild : public ISlateIMTypeChecking
+{
+};
+
+struct FSlateIMChild : private TUnion<TYPE_OF_NULLPTR, TSharedRef<SWidget>, TSharedRef<ISlateIMChild>>
+{
+	FSlateIMChild(TYPE_OF_NULLPTR = nullptr) : TUnion(nullptr)
+	{}
+	FSlateIMChild(const TSharedRef<SWidget>& Widget) : TUnion(Widget)
+	{}
+	FSlateIMChild(const TSharedRef<ISlateIMChild>& Entry) : TUnion(Entry)
+	{}
+
+	TSharedPtr<SWidget> GetWidget() const
+	{
+		if (HasSubtype<TSharedRef<SWidget>>())
+		{
+			return GetSubtype<TSharedRef<SWidget>>();
+		}
+		
+		return nullptr;
+	}
+	
+	TSharedRef<SWidget> GetWidgetRef() const
+	{
+		TSharedPtr<SWidget> Widget = GetWidget();
+		return Widget ? Widget.ToSharedRef() : SNullWidget::NullWidget;
+	}
+
+	template<typename WidgetType>
+	TSharedPtr<WidgetType> GetWidget() const
+	{
+		TSharedPtr<SWidget> Widget = GetWidget();
+		
+		if (Widget && Widget->GetWidgetClass().GetWidgetType() == WidgetType::StaticWidgetClass().GetWidgetType())
+		{
+			return StaticCastSharedPtr<WidgetType>(Widget);
+		}
+		return nullptr;
+	}
+
+	TSharedPtr<ISlateIMChild> GetChild() const
+	{
+		if (HasSubtype<TSharedRef<ISlateIMChild>>())
+		{
+			return GetSubtype<TSharedRef<ISlateIMChild>>();
+		}
+		
+		return nullptr;
+	}
+
+	template<typename ChildType>
+	TSharedPtr<ChildType> GetChild() const
+	{
+		TSharedPtr<ISlateIMChild> Child = GetChild();
+		if (Child && Child->IsA<ChildType>())
+		{
+			return StaticCastSharedPtr<ChildType>(Child);
+		}
+
+		return nullptr;
+	}
+};
\ No newline at end of file
diff --git a/Engine/Plugins/Experimental/SlateIM/Source/SlateIM/Private/Misc/SlateIMExamples.cpp b/Engine/Plugins/Experimental/SlateIM/Source/SlateIM/Private/Misc/SlateIMExamples.cpp
new file mode 100644
--- /dev/null
+++ b/Engine/Plugins/Experimental/SlateIM/Source/SlateIM/Private/Misc/SlateIMExamples.cpp
@@ -0,0 +1,621 @@
+﻿// Copyright Epic Games, Inc. All Rights Reserved.
+
+#if WITH_SLATEIM_EXAMPLES
+#include "SlateIMExamples.h"
+
+#include "Misc/App.h"
+#include "ProfilingDebugging/ScopedTimers.h"
+#include "SlateIM.h"
+#include "SlateIMLogging.h"
+#include "Styling/AppStyle.h"
+#include "Styling/StyleColors.h"
+
+#if WITH_EDITOR
+#include "SLevelViewport.h"
+#endif
+#if WITH_ENGINE
+#include "Engine/Engine.h"
+#include "Engine/Texture2D.h"
+#endif
+
+void FSlateIMTestWidget::Draw()
+{
+	double LastTime = CurrentTime;
+	CurrentTime = 0;
+	FScopedDurationTimer Timer(CurrentTime);
+
+	TimeSinceLastUpdate += FApp::GetDeltaTime();
+
+	SCOPED_NAMED_EVENT_TEXT("FSlateIMTestWidget::Draw", FColorList::Goldenrod);
+	SlateIM::BeginBorder(FAppStyle::GetBrush("ToolPanel.GroupBorder"));
+	// Basic perf measurement, outside of the scrollbox so that it's "pinned"
+	if (TimeSinceLastUpdate > 0.5f)
+	{
+		TimeText = FString::Printf(TEXT("%.3f ms"), LastTime * 1000);
+		TimeSinceLastUpdate = 0;
+	}
+	SlateIM::HAlign(HAlign_Right);
+	SlateIM::Padding({ 5.f, 5.f, 5.f, 0.f });
+	SlateIM::Text(TimeText);
+	
+	// Full root scrollbox
+	SlateIM::Fill();
+	SlateIM::HAlign(HAlign_Fill);
+	SlateIM::VAlign(VAlign_Fill);
+	SlateIM::Padding({ 5.f });
+	SlateIM::BeginScrollBox();
+	{
+		{
+			SCOPED_NAMED_EVENT_TEXT("Input Widget Examples", FColorList::Goldenrod);
+			// Button Examples
+			constexpr FStringView ClickText = TEXT("Click Me!");
+			if (SlateIM::Button(ClickText))
+			{
+				UE_LOG(LogSlateIM, Log, TEXT("Button was clicked"));
+			}
+
+			SlateIM::HAlign(HAlign_Fill);
+			if (SlateIM::Button(TEXT("Filled Button")))
+			{
+				UE_LOG(LogSlateIM, Log, TEXT("Button was clicked"));
+			}
+
+			// EditableText Example
+			if (SlateIM::EditableText(TypedText, TEXT("Type In Here")))
+			{
+				UE_LOG(LogSlateIM, Log, TEXT("Editable Text Changed: %s"), *TypedText);
+			}
+
+			// CheckBox examples
+			SlateIM::CheckBox(TEXT("Check Box"), CheckState);
+			if (CheckState)
+			{
+				SlateIM::Text(TEXT("Basic Text"));
+				SlateIM::Text(TEXT("Text With Color"), FLinearColor::Green);
+				SlateIM::Text(TEXT("Text With style color"), FStyleColors::Primary);
+			}
+		
+			if (SlateIM::CheckBox(TEXT("Undetermined Check Box"), CheckStateEnum))
+			{
+				UE_LOG(LogSlateIM, Log, TEXT("Check Box State Changed"));
+			}
+
+			// ComboBox example
+			if (SlateIM::ComboBox(ComboBoxItems, SelectedItemIndex))
+			{
+				UE_LOG(LogSlateIM, Log, TEXT("Combo Box Item %s chosen"), *ComboBoxItems[SelectedItemIndex]);
+			}
+		}
+
+		{
+			SCOPED_NAMED_EVENT_TEXT("Center Texture and Buttons", FColorList::Goldenrod);
+			// Centered alignment example
+			SlateIM::HAlign(HAlign_Center);
+			SlateIM::BeginHorizontalStack();
+
+#if WITH_ENGINE
+			// Texture Example
+			const int32 CurrentSeconds = FApp::GetCurrentTime();
+			SlateIM::VAlign(VAlign_Center);
+			UTexture2D* Texture = ((CurrentSeconds % 2) == 0)
+				? GreenIcon.LoadSynchronous()
+				: RedIcon.LoadSynchronous();
+			SlateIM::Image(Texture);
+#endif
+		
+			SlateIM::BeginVerticalStack();
+			SlateIM::Button(TEXT("Button 1"));
+			SlateIM::Button(TEXT("Button 2"));
+			SlateIM::Button(TEXT("Button 3"));
+			SlateIM::EndVerticalStack();
+
+			SlateIM::BeginVerticalStack();
+			SlateIM::Button(TEXT("Button 4"));
+			SlateIM::Button(TEXT("Button 5"));
+			SlateIM::Button(TEXT("Button 6"));
+			SlateIM::EndVerticalStack();
+
+			// SelectionList example
+			SlateIM::SelectionList(ComboBoxItems, SelectedItem, false);
+
+			SlateIM::EndHorizontalStack();
+
+			{
+				SCOPED_NAMED_EVENT_TEXT("Style image examples", FColorList::Goldenrod);
+				// Style image examples
+				SlateIM::BeginHorizontalStack();
+				SlateIM::Image("AppIcon");
+				SlateIM::Padding(FMargin(20, 10, 0, 0));
+				SlateIM::Image("Icons.ErrorWithColor");
+				SlateIM::Padding(FMargin(SliderVal, 10, 0, 0));
+				SlateIM::Image("Icons.WarningWithColor");
+				SlateIM::Padding(FMargin(SliderVal, 10, 0, 0));
+				SlateIM::Image("Icons.InfoWithColor");
+				SlateIM::Padding(FMargin(SliderVal, 10, 0, 0));
+				SlateIM::Image("Icons.SuccessWithColor");
+				SlateIM::EndHorizontalStack();
+			}
+		}
+
+		// Slider, ProgressBar, SpinBox
+		SlateIM::HAlign(HAlign_Fill);
+		SlateIM::BeginHorizontalStack();
+		{
+			SCOPED_NAMED_EVENT_TEXT("Slider, ProgressBar, SpinBox", FColorList::Goldenrod);
+			SlateIM::BeginVerticalStack();
+			{
+				// Slider example
+				if (SlateIM::Slider(SliderVal, 0, SliderMax, 1))
+				{
+					UE_LOG(LogSlateIM, Log, TEXT("Slider Value Changed [%f]"), SliderVal);
+				}
+
+				// ProgressBar example
+				SlateIM::ProgressBar(IntValue / static_cast<float>(IntMax));
+
+				// SpinBox Examples
+				{
+					SlateIM::Padding(FMargin(0, 10, 0, 5));
+					SlateIM::SpinBox(SliderVal, 0.0f, SliderMax);
+
+					SlateIM::Padding(FMargin(0, 10, 0, 5));
+					SlateIM::SpinBox(IntValue, 0, IntMax);
+				}
+			}
+			SlateIM::EndVerticalStack();
+
+			SlateIM::Fill();
+			SlateIM::HAlign(HAlign_Fill);
+			SlateIM::VAlign(VAlign_Fill);
+			if (SlateIM::Button(TEXT("Reset Values")))
+			{
+				IntValue = 50;
+				SliderVal = 5.f;
+			}
+		}
+		SlateIM::EndHorizontalStack();
+
+		{
+			SCOPED_NAMED_EVENT_TEXT("ToolTip example", FColorList::Goldenrod);
+			// ToolTip example
+			SlateIM::SetToolTip(TEXT("This Is a Tool Tip"));
+			SlateIM::BeginHorizontalStack();
+			SlateIM::Text(TEXT("Tool Tip Testing:"));
+			SlateIM::Image("AppIcon");
+			SlateIM::EndHorizontalStack();
+		}
+		
+		{
+			SCOPED_NAMED_EVENT_TEXT("DisabledState example", FColorList::Goldenrod);
+			// DisabledState example
+			SlateIM::CheckBox(TEXT("Disable Everything Below Me"), bShouldBeDisabled);
+			if (bShouldBeDisabled)
+			{
+				SlateIM::BeginDisabledState();
+			}
+		}
+
+		// ContextMenu examples
+		{
+			SCOPED_NAMED_EVENT_TEXT("ContextMenu examples", FColorList::Goldenrod);
+			SlateIM::Text(TEXT("Context Menu Test"));
+			SlateIM::BeginContextMenuAnchor();
+			SlateIM::Text(TEXT("Right Click here to show a menu"));
+
+			// This part is only shown if the menu is open
+			SlateIM::AddMenuSection(TEXT("Menu Section 1"));
+			if (SlateIM::AddMenuButton(TEXT("Menu Item 1"), TEXT("Menu Item Tool Tip 1")))
+			{
+				UE_LOG(LogSlateIM, Log, TEXT("Menu Item One menu option clicked"));
+			}
+
+			SlateIM::AddMenuButton(TEXT("Menu Item 2"), TEXT("Menu Item Tool Tip 2"));
+			SlateIM::AddMenuButton(TEXT("Menu Item 3"), TEXT("Menu Item Tool Tip 3"));
+			SlateIM::AddMenuButton(TEXT("Menu Item 4"), TEXT("Menu Item Tool Tip 4"));
+
+			SlateIM::AddMenuCheckButton(TEXT("Menu Item With Check"), MenuCheckState, TEXT("Click to toggle check"));
+
+			if (SlateIM::AddMenuToggleButton(TEXT("Menu Item With Toggle"), MenuToggleState, TEXT("Toggle this box")))
+			{
+				UE_LOG(LogSlateIM, Log, TEXT("Menu Item With Toggle clicked"));
+			}
+
+			SlateIM::AddMenuSeparator();
+			SlateIM::BeginSubMenu(TEXT("Sub Menu"));
+			SlateIM::AddMenuButton(TEXT("SubMenu Item 1"), TEXT("Menu Item Tool Tip 1"));
+			SlateIM::AddMenuButton(TEXT("SubMenu Item 2"), TEXT("Menu Item Tool Tip 2"));
+			SlateIM::AddMenuButton(TEXT("SubMenu Item 3"), TEXT("Menu Item Tool Tip 3"));
+			SlateIM::AddMenuButton(TEXT("SubMenu Item 4"), TEXT("Menu Item Tool Tip 4"));
+			SlateIM::EndSubMenu();
+			SlateIM::EndContextMenuAnchor();
+		}
+	
+
+		// Modal Examples
+		{
+			SCOPED_NAMED_EVENT_TEXT("Modal Examples", FColorList::Goldenrod);
+			SlateIM::Text(TEXT("Open Modal Dialog of Type:"));
+
+			// Wrap example
+			SlateIM::HAlign(HAlign_Fill);
+			SlateIM::BeginHorizontalWrap();
+			if (SlateIM::Button(TEXT("Ok")))
+			{
+				DialogResult = SlateIM::ModalDialog(EAppMsgType::Ok, TEXT("Ok?"));
+			}
+			if (SlateIM::Button(TEXT("YesNo")))
+			{
+				DialogResult = SlateIM::ModalDialog(EAppMsgType::YesNo, TEXT("YesNo?"));
+			}
+			if (SlateIM::Button(TEXT("OkCancel")))
+			{
+				DialogResult = SlateIM::ModalDialog(EAppMsgType::OkCancel, TEXT("OkCancel?"));
+			}
+			if (SlateIM::Button(TEXT("YesNoCancel")))
+			{
+				DialogResult = SlateIM::ModalDialog(EAppMsgType::YesNoCancel, TEXT("YesNoCancel?"));
+			}
+			if (SlateIM::Button(TEXT("CancelRetryContinue")))
+			{
+				DialogResult = SlateIM::ModalDialog(EAppMsgType::CancelRetryContinue, TEXT("CancelRetryContinue?"));
+			}
+			if (SlateIM::Button(TEXT("YesNoYesAllNoAll")))
+			{
+				DialogResult = SlateIM::ModalDialog(EAppMsgType::YesNoYesAllNoAll, TEXT("YesNoYesAllNoAll?"));
+			}
+			if (SlateIM::Button(TEXT("YesNoYesAllNoAllCancel")))
+			{
+				DialogResult = SlateIM::ModalDialog(EAppMsgType::YesNoYesAllNoAllCancel, TEXT("YesNoYesAllNoAllCancel?"));
+			}
+			if (SlateIM::Button(TEXT("YesNoYesAll")))
+			{
+				DialogResult = SlateIM::ModalDialog(EAppMsgType::YesNoYesAll, TEXT("YesNoYesAll?"));
+			}
+			SlateIM::EndHorizontalWrap();
+
+			if (DialogResult.IsSet())
+			{
+				SlateIM::Text(TEXT("Dialog Result:"));
+				switch (DialogResult.GetValue())
+				{
+				case EAppReturnType::No:
+					SlateIM::Text(TEXT("No"));
+					break;
+				case EAppReturnType::Yes:
+					SlateIM::Text(TEXT("Yes"));
+					break;
+				case EAppReturnType::YesAll:
+					SlateIM::Text(TEXT("Yes to All"));
+					break;
+				case EAppReturnType::NoAll:
+					SlateIM::Text(TEXT("No to All"));
+					break;
+				case EAppReturnType::Cancel:
+					SlateIM::Text(TEXT("Cancel"));
+					break;
+				case EAppReturnType::Ok:
+					SlateIM::Text(TEXT("Ok"));
+					break;
+				case EAppReturnType::Retry:
+					SlateIM::Text(TEXT("Retry"));
+					break;
+				case EAppReturnType::Continue:
+					SlateIM::Text(TEXT("Continue"));
+					break;
+				default:
+					SlateIM::Text(TEXT("UNHANDLED RESULT"));
+					break;
+				}
+			
+				if (SlateIM::Button(TEXT("Reset")))
+				{
+					DialogResult.Reset();
+				}
+			}
+			else
+			{
+				SlateIM::Text(TEXT("No Dialog Result"));
+			}
+		}
+		
+		SlateIM::Text(TEXT("Scroll Bar Test"));
+		SlateIM::BeginHorizontalStack();
+
+		SlateIM::Text(TEXT("Num Items "));
+		SlateIM::MinWidth(50.f);
+		if (SlateIM::EditableText(NumItemsText))
+		{
+			LiveNumItems = FCString::Atoi(*NumItemsText);
+		}
+
+		if (SlateIM::Button(TEXT("Regenerate")))
+		{
+			NumItems = FCString::Atoi(*NumItemsText);
+		}
+
+		int32 NumTableItems = NumItems;
+		SlateIM::CheckBox(TEXT("Live Update Table?"), bShouldLiveUpdateTable);
+		if (bShouldLiveUpdateTable)
+		{
+			NumTableItems = LiveNumItems;
+		}
+		SlateIM::EndHorizontalStack();
+
+		SlateIM::VAlign(VAlign_Fill);
+		SlateIM::HAlign(HAlign_Fill);
+		SlateIM::BeginHorizontalStack();
+
+		// ScrollBox example
+		{
+			SCOPED_NAMED_EVENT_TEXT("ScrollBox example", FColorList::Goldenrod);
+			SlateIM::AutoSize();
+			SlateIM::MaxHeight(200.f);
+			SlateIM::BeginScrollBox();
+			for (int32 i = 0; i < NumItems; ++i)
+			{
+				// New row
+				SlateIM::BeginHorizontalStack();
+				{
+					SlateIM::Padding(FMargin(5, 0));
+					SlateIM::VAlign(VAlign_Center);// Centers the button in the row
+					SlateIM::Text(FString::Printf(TEXT("Item %d/%d"), i + 1, NumItems), FColor::MakeRedToGreenColorFromScalar(static_cast<float>(i) / NumItems));
+					SlateIM::Padding(FMargin(0));
+
+					if (SlateIM::Button(TEXT("Click")))
+					{
+						UE_LOG(LogSlateIM, Log, TEXT("Button %d clicked"), i + 1);
+					}
+				}
+				SlateIM::EndHorizontalStack();
+			}
+			SlateIM::EndScrollBox();
+		}
+
+		// Spacer Example
+		SlateIM::Spacer({ 20, 1 });
+
+		// Table Example
+		{
+			SCOPED_NAMED_EVENT_TEXT("Table Example", FColorList::Goldenrod);
+			SlateIM::Fill();
+			SlateIM::HAlign(HAlign_Fill);
+			SlateIM::MaxHeight(200.f);
+			SlateIM::BeginTable();
+			SlateIM::AddTableColumn(TEXT("Item"));
+			SlateIM::HAlign(HAlign_Center);
+			SlateIM::InitialTableColumnWidth(80.f);
+			SlateIM::AddTableColumn(TEXT("Button"));
+			for (int32 i = 0; i < NumTableItems; ++i)
+			{
+				if (SlateIM::NextTableCell())
+				{
+					SlateIM::Padding(FMargin(5, 0));
+					SlateIM::Fill();
+					SlateIM::VAlign(VAlign_Center);
+					SlateIM::Text(FString::Printf(TEXT("Item %d/%d"), i + 1, NumTableItems), FColor::MakeRedToGreenColorFromScalar(static_cast<float>(i) / NumTableItems));
+				}
+
+				if (SlateIM::NextTableCell())
+				{
+					SlateIM::Padding(FMargin(0));
+					SlateIM::HAlign(HAlign_Center);
+					if (SlateIM::Button(TEXT("Click")))
+					{
+						UE_LOG(LogSlateIM, Log, TEXT("Table Button %d clicked"), i + 1);
+					}
+				}
+			}
+			SlateIM::EndTable();
+		}
+		SlateIM::EndHorizontalStack();
+
+		// Tree Example
+		{
+			SCOPED_NAMED_EVENT_TEXT("Tree Example", FColorList::Goldenrod);
+			SlateIM::MinWidth(500.f);
+			SlateIM::MinHeight(200.f);
+			SlateIM::MaxHeight(200.f);
+			SlateIM::VAlign(VAlign_Fill);
+			SlateIM::BeginTable();
+			SlateIM::AddTableColumn(TEXT("Name"));
+			SlateIM::AddTableColumn(TEXT("Type"));
+			{
+				if (SlateIM::NextTableCell())
+				{
+					SlateIM::VAlign(VAlign_Center);
+					SlateIM::Text(TEXT("Antarctica"));
+				}
+				if (SlateIM::NextTableCell())
+				{
+					SlateIM::VAlign(VAlign_Center);
+					SlateIM::Text(TEXT("Continent"));
+				}
+				
+				if (SlateIM::NextTableCell())
+				{
+					SlateIM::VAlign(VAlign_Center);
+					SlateIM::Text(TEXT("North America"));
+				}
+				if (SlateIM::NextTableCell())
+				{
+					SlateIM::VAlign(VAlign_Center);
+					SlateIM::Text(TEXT("Continent"));
+				}
+				if (SlateIM::BeginTableRowChildren())
+				{
+					if (SlateIM::NextTableCell())
+					{
+						SlateIM::VAlign(VAlign_Center);
+						SlateIM::Text(TEXT("Canada"));
+					}
+					if (SlateIM::NextTableCell())
+					{
+						SlateIM::VAlign(VAlign_Center);
+						SlateIM::Text(TEXT("Country"));
+					}
+					if (SlateIM::BeginTableRowChildren())
+					{
+						if (SlateIM::NextTableCell())
+						{
+							SlateIM::VAlign(VAlign_Center);
+							SlateIM::Text(TEXT("British Columbia"));
+						}
+						if (SlateIM::NextTableCell())
+						{
+							SlateIM::VAlign(VAlign_Center);
+							SlateIM::Text(TEXT("Province"));
+						}
+						if (SlateIM::BeginTableRowChildren())
+						{
+							if (SlateIM::NextTableCell())
+							{
+								SlateIM::VAlign(VAlign_Center);
+								SlateIM::Text(TEXT("Vancouver"));
+							}
+							if (SlateIM::NextTableCell())
+							{
+								SlateIM::VAlign(VAlign_Center);
+								SlateIM::Text(TEXT("City"));
+							}
+						}
+						SlateIM::EndTableRowChildren();
+					
+						if (SlateIM::NextTableCell())
+						{
+							SlateIM::VAlign(VAlign_Center);
+							SlateIM::Text(TEXT("Quebec"));
+						}
+						if (SlateIM::NextTableCell())
+						{
+							SlateIM::VAlign(VAlign_Center);
+							SlateIM::Text(TEXT("Province"));
+						}
+						if (SlateIM::BeginTableRowChildren())
+						{
+							if (SlateIM::NextTableCell())
+							{
+								SlateIM::VAlign(VAlign_Center);
+								SlateIM::Text(TEXT("Montreal"));
+							}
+							if (SlateIM::NextTableCell())
+							{
+								SlateIM::VAlign(VAlign_Center);
+								SlateIM::Text(TEXT("City"));
+							}
+						}
+						SlateIM::EndTableRowChildren();
+					}
+					SlateIM::EndTableRowChildren();
+
+					if (SlateIM::NextTableCell())
+					{
+						SlateIM::VAlign(VAlign_Center);
+						SlateIM::Text(TEXT("United States"));
+					}
+					if (SlateIM::NextTableCell())
+					{
+						SlateIM::VAlign(VAlign_Center);
+						SlateIM::Text(TEXT("Country"));
+					}
+					if (SlateIM::BeginTableRowChildren())
+					{
+						if (SlateIM::NextTableCell())
+						{
+							SlateIM::VAlign(VAlign_Center);
+							SlateIM::Text(TEXT("North Carolina"));
+						}
+						if (SlateIM::NextTableCell())
+						{
+							SlateIM::VAlign(VAlign_Center);
+							SlateIM::Text(TEXT("State"));
+						}
+						if (SlateIM::BeginTableRowChildren())
+						{
+							if (SlateIM::NextTableCell())
+							{
+								SlateIM::VAlign(VAlign_Center);
+								SlateIM::Text(TEXT("Cary"));
+							}
+							if (SlateIM::NextTableCell())
+							{
+								SlateIM::VAlign(VAlign_Center);
+								SlateIM::Text(TEXT("City"));
+							}
+						}
+						SlateIM::EndTableRowChildren();
+						
+						if (SlateIM::NextTableCell())
+						{
+							SlateIM::VAlign(VAlign_Center);
+							SlateIM::Text(TEXT("Washington"));
+						}
+						if (SlateIM::NextTableCell())
+						{
+							SlateIM::VAlign(VAlign_Center);
+							SlateIM::Text(TEXT("State"));
+						}
+						if (SlateIM::BeginTableRowChildren())
+						{
+							if (SlateIM::NextTableCell())
+							{
+								SlateIM::VAlign(VAlign_Center);
+								SlateIM::Text(TEXT("Bellevue"));
+							}
+							if (SlateIM::NextTableCell())
+							{
+								SlateIM::VAlign(VAlign_Center);
+								SlateIM::Text(TEXT("City"));
+							}
+						}
+						SlateIM::EndTableRowChildren();
+					}
+					SlateIM::EndTableRowChildren();
+				}
+				SlateIM::EndTableRowChildren();
+			}
+			
+			SlateIM::EndTable();
+		}
+
+		if (bShouldBeDisabled)
+		{
+			SlateIM::EndDisabledState();
+		}
+	}
+	SlateIM::EndScrollBox();
+	SlateIM::EndBorder();
+}
+
+void FSlateIMTestWindowWidget::DrawWindow(float DeltaTime)
+{
+	TestWidget.Draw();
+}
+
+#if WITH_ENGINE
+void FSlateIMTestViewportWidget::DrawWidget(float DeltaTime)
+{
+	if (GEngine && GEngine->GameViewport)
+	{
+		if (SlateIM::BeginViewportRoot("SlateIMTestSuiteViewport", GEngine->GameViewport))
+		{
+			TestWidget.Draw();
+		}
+		SlateIM::EndRoot();
+	}
+#if WITH_EDITOR
+	else if (GCurrentLevelEditingViewportClient)
+	{
+		TSharedPtr<SLevelViewport> LevelViewport = StaticCastSharedPtr<SLevelViewport>(GCurrentLevelEditingViewportClient->GetEditorViewportWidget());
+		if (LevelViewport.IsValid())
+		{
+			if (SlateIM::BeginViewportRoot("SlateIMTestSuiteViewport", LevelViewport, FVector2f(0, 25), FVector2f(.5f, .75f)))
+			{
+				TestWidget.Draw();
+			}
+			SlateIM::EndRoot();
+		}
+	}
+#endif
+}
+#endif // WITH_ENGINE
+#endif // WITH_SLATEIM_EXAMPLES
diff --git a/Engine/Plugins/Experimental/SlateIM/Source/SlateIM/Private/Misc/SlateIMExamples.h b/Engine/Plugins/Experimental/SlateIM/Source/SlateIM/Private/Misc/SlateIMExamples.h
new file mode 100644
--- /dev/null
+++ b/Engine/Plugins/Experimental/SlateIM/Source/SlateIM/Private/Misc/SlateIMExamples.h
@@ -0,0 +1,72 @@
+﻿// Copyright Epic Games, Inc. All Rights Reserved.
+
+#pragma once
+
+#if WITH_SLATEIM_EXAMPLES
+#include "SlateIMWidgetBase.h"
+#include "Styling/SlateTypes.h"
+
+class UTexture2D;
+
+class FSlateIMTestWidget
+{
+public:
+	void Draw();
+	
+private:
+	double CurrentTime = 0;
+	double TimeSinceLastUpdate = 0;
+	FString TimeText;
+	FString TypedText;
+	bool CheckState = false;
+	ECheckBoxState CheckStateEnum = ECheckBoxState::Undetermined;
+	int32 SelectedItemIndex = 0;
+	float SliderVal = 5.0f;
+	float SliderMax = 20.f;
+	int32 SelectedItem = INDEX_NONE;
+	int32 IntValue = 50;
+	int32 IntMax = 100;
+	bool bShouldBeDisabled = false;
+	bool MenuCheckState = true;
+	bool MenuToggleState = true;
+	TOptional<EAppReturnType::Type> DialogResult;
+	int32 NumItems = 10;
+	int32 LiveNumItems = NumItems;
+	FString NumItemsText = FString::FromInt(NumItems);
+	bool bShouldLiveUpdateTable = false;
+	const TArray<FString> ComboBoxItems = { TEXT("Option 1"), TEXT("Option 2"), TEXT("Option 3"), TEXT("Option 4") };
+
+#if WITH_ENGINE
+	TSoftObjectPtr<UTexture2D> RedIcon = TSoftObjectPtr<UTexture2D>(FSoftObjectPath(TEXT("/Engine/EngineResources/AICON-Red.AICON-Red")));
+	TSoftObjectPtr<UTexture2D> GreenIcon = TSoftObjectPtr<UTexture2D>(FSoftObjectPath(TEXT("/Engine/EngineResources/AICON-Green.AICON-Green")));
+#endif
+};
+
+class FSlateIMTestWindowWidget : public FSlateIMWindowBase
+{
+public:
+	FSlateIMTestWindowWidget(const TCHAR* Command, const TCHAR* CommandHelp)
+		: FSlateIMWindowBase(TEXT("SlateIM Test Suite"), FVector2f(700, 900), Command, CommandHelp)
+	{}
+	
+	virtual void DrawWindow(float DeltaTime) override;
+
+private:
+	FSlateIMTestWidget TestWidget;
+};
+
+#if WITH_ENGINE
+class FSlateIMTestViewportWidget : public FSlateIMWidgetWithCommandBase
+{
+public:
+	FSlateIMTestViewportWidget(const TCHAR* Command, const TCHAR* CommandHelp)
+		: FSlateIMWidgetWithCommandBase(Command, CommandHelp)
+	{}
+	
+	virtual void DrawWidget(float DeltaTime) override;
+
+private:
+	FSlateIMTestWidget TestWidget;
+};
+#endif // WITH_ENGINE
+#endif // WITH_SLATEIM_EXAMPLES
diff --git a/Engine/Plugins/Experimental/SlateIM/Source/SlateIM/Private/Misc/SlateIMLogging.h b/Engine/Plugins/Experimental/SlateIM/Source/SlateIM/Private/Misc/SlateIMLogging.h
new file mode 100644
--- /dev/null
+++ b/Engine/Plugins/Experimental/SlateIM/Source/SlateIM/Private/Misc/SlateIMLogging.h
@@ -0,0 +1,7 @@
+﻿// Copyright Epic Games, Inc. All Rights Reserved.
+
+#pragma once
+
+#include "Logging/LogMacros.h"
+
+DECLARE_LOG_CATEGORY_EXTERN(LogSlateIM, Log, All);
diff --git a/Engine/Plugins/Experimental/SlateIM/Source/SlateIM/Private/Misc/SlateIMManager.cpp b/Engine/Plugins/Experimental/SlateIM/Source/SlateIM/Private/Misc/SlateIMManager.cpp
new file mode 100644
--- /dev/null
+++ b/Engine/Plugins/Experimental/SlateIM/Source/SlateIM/Private/Misc/SlateIMManager.cpp
@@ -0,0 +1,382 @@
+﻿// Copyright Epic Games, Inc. All Rights Reserved.
+
+#include "SlateIMManager.h"
+
+#include "Containers/SImContextMenuAnchor.h"
+#include "Containers/SImStackBox.h"
+#include "Framework/Application/SlateApplication.h"
+#include "Misc/App.h"
+#include "Roots/ISlateIMRoot.h"
+#include "SlateIMLogging.h"
+#include "SlateIM.h"
+
+namespace SlateIM
+{
+	FRootNode::FRootNode(FName InRootName, TSharedPtr<ISlateIMRoot> InRootWidget, bool InRootState)
+			: RootContainer(nullptr)
+			, RootWidget(InRootWidget)
+			, RootName(InRootName)
+			, RootState(InRootState)
+	{}
+
+	FWidgetHash FRootNode::GetWidgetHash() const
+	{
+		return WidgetTreeDataHash.Num() > CurrentWidgetIndex ? WidgetTreeDataHash[CurrentWidgetIndex] : FWidgetHash();
+	}
+
+	void FRootNode::SetDataHash(FXxHash64 InDataHash)
+	{
+		if (WidgetTreeDataHash.Num() <= CurrentWidgetIndex)
+		{
+			WidgetTreeDataHash.Emplace(FXxHash64(), InDataHash);
+			check(WidgetTreeDataHash.Num() - 1 == CurrentWidgetIndex);
+		}
+		else
+		{
+			WidgetTreeDataHash[CurrentWidgetIndex].DataHash = InDataHash;
+		}
+	}
+
+	void FRootNode::SetAlignmentHash(FXxHash64 AlignmentHash)
+	{
+		UE_LOG(LogSlateIM, Verbose, TEXT("SetAlignmentHash - CurrentWidgetIndex [%d] | WidgetTreeDataHash.Num() [%d]"), CurrentWidgetIndex, WidgetTreeDataHash.Num());
+		if (WidgetTreeDataHash.Num() <= CurrentWidgetIndex)
+		{
+			WidgetTreeDataHash.Emplace(AlignmentHash, FXxHash64());
+			ensure(WidgetTreeDataHash.Num() - 1 == CurrentWidgetIndex);
+		}
+		else
+		{
+			WidgetTreeDataHash[CurrentWidgetIndex].AlignmentHash = AlignmentHash;
+		}
+	}
+
+	void FRootNode::SetNextToolTip(const FStringView& InNextToolTip)
+	{
+		if (!InNextToolTip.IsEmpty())
+		{
+			CurrentToolTip = InNextToolTip;
+		}
+		else
+		{
+			CurrentToolTip.Empty();
+		}
+	}
+
+	TUniquePtr<FSlateIMManager> FSlateIMManager::Instance;
+
+	void FSlateIMManager::Initialize()
+	{
+		if (FSlateApplication::IsInitialized() && !Instance)
+		{
+			Instance = MakeUnique<FSlateIMManager>();
+		}
+	}
+
+	FSlateIMManager& FSlateIMManager::Get()
+	{
+		check(Instance.IsValid());
+
+		return *Instance;
+	}
+
+	FSlateIMManager::FSlateIMManager()
+#if WITH_SLATEIM_EXAMPLES
+			: TestWindowWidget(TEXT("SlateIM.ToggleTestSuiteWindow"), TEXT("Toggles the Slate immediate mode test suite window which demonstrates the capabilities of the Slate immediate mode api"))
+#if WITH_ENGINE
+			, TestViewportWidget(TEXT("SlateIM.ToggleTestSuiteViewport"), TEXT("Toggles the Slate immediate mode test suite in the main viewport (PIE or LevelEditor), demonstrating the capabilities of the Slate immediate mode api"))
+#endif
+#endif
+	{
+		if (FSlateApplication::IsInitialized())
+		{
+			FSlateApplication::Get().OnPostTick().AddRaw(this, &FSlateIMManager::Tick);
+			FSlateApplication::Get().OnPreShutdown().AddRaw(this, &FSlateIMManager::OnSlateShutdown);
+		}
+	}
+
+	FSlateIMManager::~FSlateIMManager()
+	{
+		if (FSlateApplication::IsInitialized())
+		{
+			FSlateApplication::Get().OnPreShutdown().RemoveAll(this);
+			FSlateApplication::Get().OnPostTick().RemoveAll(this);
+		}
+	}
+
+	FRootNode& FSlateIMManager::AddRoot(const FName WindowId, TSharedPtr<ISlateIMRoot> NewRoot)
+	{
+		return RootMap.Emplace(WindowId, FRootNode(WindowId, NewRoot, true));
+	}
+
+	void FSlateIMManager::BeginRoot(FName RootName)
+	{
+		if (!CanUpdateSlateIM())
+		{
+			return;
+		}
+		
+		FPlatformMisc::BeginNamedEvent(FColorList::Goldenrod, *RootName.ToString());
+		TRACE_CPUPROFILER_EVENT_MANUAL_START(RootName);
+		
+		checkf(CurrentRoot == nullptr, TEXT("Cannot begin an new SlateIM root while one is already being built. Call SlateIM::EndRoot() to end current root and begin a new one"));
+
+		check(MemMark == nullptr);
+
+		MemMark = MakeUnique<FMemMark>(FMemStack::Get());
+
+		CurrentRoot = &RootMap.FindChecked(RootName);
+		CurrentRoot->LastAccessTime = FApp::GetCurrentTime();
+
+		CurrentRoot->bCurrentEnabledState = true;
+		CurrentRoot->bActivatedThisFrame = true;
+		CurrentRoot->CurrentContainerStack.Reset();
+		CurrentRoot->CurrentMenuAnchorStack.Reset();
+
+		CurrentRoot->CurrentWidgetIndex = 0;
+
+		if (CurrentRoot->RootContainer.IsValid())
+		{
+			CurrentRoot->CurrentContainerStack.Push(CurrentRoot->RootContainer);
+		}
+	}
+
+	void FSlateIMManager::EndRoot()
+	{
+		if (!CanUpdateSlateIM())
+		{
+			return;
+		}
+		
+		checkf(CurrentRoot != nullptr, TEXT("Called SlateIM::EndRoot() without an active root"));
+
+		if (CurrentRoot->CurrentContainerStack.Num() > 0)
+		{
+			// TODO - Is this an error caused by a missing Pop?
+			RemoveUnusedChildren(CurrentRoot->CurrentContainerStack.Top());
+		}
+
+		CurrentRoot->CurrentContainerStack.Reset();
+		CurrentRoot->CurrentMenuAnchorStack.Reset();
+
+		const int32 FirstUnusedIndex = CurrentRoot->CurrentWidgetIndex + 1;
+		if (FirstUnusedIndex < CurrentRoot->WidgetTreeDataHash.Num())
+		{
+			CurrentRoot->WidgetTreeDataHash.RemoveAt(FirstUnusedIndex, CurrentRoot->WidgetTreeDataHash.Num() - FirstUnusedIndex, EAllowShrinking::No);
+		}
+
+		CurrentRoot->CurrentWidgetIndex = 0;
+		CurrentRoot = nullptr;
+
+		MemMark.Reset();
+
+		TRACE_CPUPROFILER_EVENT_MANUAL_END();
+		FPlatformMisc::EndNamedEvent();
+	}
+
+	void FSlateIMManager::PushContainer(FContainerNode&& Node)
+	{
+		// Container widget should exist in the parent
+		CurrentRoot->CurrentContainerStack.Push(Node);
+	}
+
+	const FContainerNode* FSlateIMManager::GetCurrentContainerNode() const
+	{
+		return CurrentRoot->CurrentContainerStack.IsEmpty() ? nullptr : &CurrentRoot->CurrentContainerStack.Top();
+	}
+
+	void FSlateIMManager::PushMenuRoot(TSharedPtr<SImContextMenuAnchor>& MenuRoot)
+	{
+		MenuRoot->Begin();
+		CurrentRoot->CurrentMenuAnchorStack.Push(MenuRoot);
+	}
+
+	void FSlateIMManager::PopMenuRoot()
+	{
+		TSharedPtr<SImContextMenuAnchor> MenuRoot = CurrentRoot->CurrentMenuAnchorStack.Pop(EAllowShrinking::No);
+		MenuRoot->End();
+	}
+
+	TSharedPtr<SImContextMenuAnchor> FSlateIMManager::GetCurrentMenuRoot() const
+	{
+		return CurrentRoot->CurrentMenuAnchorStack.Num() ? CurrentRoot->CurrentMenuAnchorStack.Top() : nullptr;
+	}
+
+	TSharedPtr<SWidget> FSlateIMManager::BeginCustomWidget(TSharedPtr<SWidget> ExpectedWidget)
+	{
+		bBuildingImmediateModeWidgets = true;
+
+		EnsureCurrentContainerNode();
+		AdvanceToNextWidget();
+
+		if (GetCurrentChild().GetWidget() == ExpectedWidget)
+		{
+			return ExpectedWidget;
+		}
+			
+		return nullptr;
+	}
+
+	void FSlateIMManager::EndWidget(bool bResetAlignmentData)
+	{
+		bBuildingImmediateModeWidgets = false;
+
+		if (bResetAlignmentData)
+		{
+			ResetAlignmentData();
+		}
+	}
+
+	void FSlateIMManager::EnsureCurrentContainerNode()
+	{
+		check(CurrentRoot != nullptr);
+		
+		// Add new vertical box (the default layout) as a container inside the current root. The current root always has one child
+		if (CurrentRoot->CurrentContainerStack.Num() == 0 && !CurrentRoot->RootContainer.IsValid())
+		{
+			UE_LOG(LogSlateIM, Log, TEXT("Adding default root container"));
+			TSharedRef<SImStackBox> VerticalBox = SNew(SImStackBox).Orientation(Orient_Vertical);
+
+			const FSlateIMSlotData AlignmentData(FMargin(0), HAlign_Fill, VAlign_Fill, false, 0, 0, 0, 0);
+
+			CurrentRoot->RootWidget->UpdateChild(VerticalBox, AlignmentData);
+
+			CurrentRoot->SetAlignmentHash(AlignmentData.Hash);
+
+			CurrentRoot->RootContainer = FContainerNode(VerticalBox);
+
+			CurrentRoot->CurrentContainerStack.Push(CurrentRoot->RootContainer);
+		}
+	}
+	
+	FSlateIMChild FSlateIMManager::GetCurrentChild() const
+	{
+		check(CurrentRoot->CurrentContainerStack.Num());
+
+		const FContainerNode& CurrentContainer = CurrentRoot->CurrentContainerStack.Top();
+		
+		return ensure(CurrentContainer.Widget) ? CurrentContainer.Widget->GetChild(CurrentContainer.LastUsedChildIndex) : nullptr;
+	}
+
+	void FSlateIMManager::UpdateCurrentChild(FSlateIMChild Child, const FSlateIMSlotData& AlignmentData)
+	{
+		FContainerNode& CurrentContainer = CurrentRoot->CurrentContainerStack.Top();
+		CurrentRoot->SetAlignmentHash(AlignmentData.Hash);
+
+		check(CurrentContainer.Widget);
+		UE_LOG(LogSlateIM, Verbose, TEXT("Updating Container: [%s] Child: [%d]"), *CurrentContainer.Widget->GetDebugName(), CurrentContainer.LastUsedChildIndex);
+		CurrentContainer.Widget->UpdateChild(Child, CurrentContainer.LastUsedChildIndex, AlignmentData);
+	}
+
+	void FSlateIMManager::AdvanceToNextWidget()
+	{
+		++CurrentRoot->CurrentWidgetIndex;
+		++CurrentRoot->CurrentContainerStack.Top().LastUsedChildIndex;
+	}
+
+	void FSlateIMManager::ActivateWidget(const FWidgetActivation& ActivationData)
+	{
+		// Setting values while building immediate mode ui causes some widgets to activate
+		// This is undesirable as it can cause feedback loops
+		if (!bBuildingImmediateModeWidgets)
+		{
+			UE_LOG(LogSlateIM, Verbose, TEXT("ActivateWidget | Root: [%s] | Container [%d] | Widget [%d]"), *ActivationData.RootName.ToString(), ActivationData.ContainerIndex, ActivationData.WidgetIndex);
+			ActivatedWidgets.Add(ActivationData);
+		}
+	}
+
+	bool FSlateIMManager::IsWidgetActivatedThisFrame() const
+	{
+		return ActivatedWidgets.Contains(GetWidgetActivationData());
+	}
+
+	FWidgetActivation FSlateIMManager::GetWidgetActivationData() const
+	{
+		const FRootNode& Root = FSlateIMManager::Get().GetCurrentRoot();
+		const int32 ContainerIndex = Root.CurrentContainerStack.Num() - 1;
+		const int32 WidgetIndex = Root.CurrentWidgetIndex;
+
+		return { Root.RootName, ContainerIndex, WidgetIndex };
+	}
+
+	void FSlateIMManager::ResetAlignmentData()
+	{
+		NextPadding.Reset();
+		NextHAlign.Reset();
+		NextVAlign.Reset();
+		NextAutoSize.Reset();
+		NextMinWidth.Reset();
+		NextMinHeight.Reset();
+		NextMaxWidth.Reset();
+		NextMaxHeight.Reset();
+	}
+
+	void FSlateIMManager::OnSlateIMModalOpened()
+	{
+		ensure(!bIsSlateIMModalOpen);
+		bIsSlateIMModalOpen = true;
+	}
+
+	void FSlateIMManager::OnSlateIMModalClosed()
+	{
+		bIsSlateIMModalOpen = false;
+	}
+
+	bool FSlateIMManager::CanUpdateSlateIM() const
+	{
+		// Disable SlateIM updates while we have an open modal
+		return !bIsSlateIMModalOpen;
+	}
+
+	void FSlateIMManager::ValidateRootName(FName RootName) const
+	{
+		checkf(!RootMap.Contains(RootName), TEXT("SlateIM: Root Name %s is not unique. All roots of Slate IM hierarchies must have a unique name"), *RootName.ToString());
+	}
+
+	void FSlateIMManager::Tick(float DeltaTime)
+	{
+		SCOPED_NAMED_EVENT_TEXT("FSlateIMManager::Tick", FColorList::Goldenrod);
+		if (!CanUpdateSlateIM())
+		{
+			return;
+		}
+
+		check(CurrentRoot == nullptr);
+
+		ActivatedWidgets.Reset();
+
+		for (auto It = RootMap.CreateIterator(); It; ++It)
+		{
+			if (It.Value().bActivatedThisFrame == false)
+			{
+				It.RemoveCurrent();
+			}
+			else
+			{
+				// Reset for next frame
+				It.Value().bActivatedThisFrame = false;
+			}
+		}
+	}
+
+	void FSlateIMManager::OnSlateShutdown()
+	{
+		ActivatedWidgets.Empty();
+		RootMap.Empty();
+	}
+
+	void FSlateIMManager::RemoveUnusedChildren(FContainerNode& Container)
+	{
+		if (Container.Widget->GetNumChildren() > Container.LastUsedChildIndex + 1)
+		{
+			TSharedPtr<ISlateIMContainer> ContainerWidget = Container.Widget;
+			ContainerWidget->RemoveUnusedChildren(Container.LastUsedChildIndex);
+		}
+	}
+
+	void FSlateIMManager::PopContainer_Internal()
+	{
+		RemoveUnusedChildren(CurrentRoot->CurrentContainerStack.Top());
+		CurrentRoot->CurrentContainerStack.Pop(EAllowShrinking::No);
+	}
+}
diff --git a/Engine/Plugins/Experimental/SlateIM/Source/SlateIM/Private/Misc/SlateIMManager.h b/Engine/Plugins/Experimental/SlateIM/Source/SlateIM/Private/Misc/SlateIMManager.h
new file mode 100644
--- /dev/null
+++ b/Engine/Plugins/Experimental/SlateIM/Source/SlateIM/Private/Misc/SlateIMManager.h
@@ -0,0 +1,277 @@
+﻿// Copyright Epic Games, Inc. All Rights Reserved.
+
+#pragma once
+
+#include "Containers/ISlateIMContainer.h"
+#include "Containers/Union.h"
+#include "Hash/xxhash.h"
+#include "Layout/Margin.h"
+#include "Roots/ISlateIMRoot.h"
+#include "SlateIMLogging.h"
+#include "SlateIMSlotData.h"
+#include "Templates/SharedPointer.h"
+#include "Templates/UniquePtr.h"
+#include "Widgets/SWidget.h"
+
+#if WITH_SLATEIM_EXAMPLES
+#include "SlateIMExamples.h"
+#endif
+
+class FMemMark;
+class ISlateIMChild;
+class ISlateIMRoot;
+class ISlateIMContainer;
+class SImContextMenuAnchor;
+class SWidget;
+
+struct FSlateIMSlotData;
+
+namespace SlateIM
+{
+	struct FContainerNode
+	{
+		FContainerNode(TSharedPtr<ISlateIMContainer> InWidget)
+			: Widget(InWidget)
+		{}
+
+		bool IsValid() const { return Widget.IsValid(); }
+
+		/** Current State */
+		TSharedPtr<ISlateIMContainer> Widget;
+		int32 LastUsedChildIndex = INDEX_NONE;
+	};
+	
+	struct FWidgetHash
+	{
+		FXxHash64 AlignmentHash;
+		FXxHash64 DataHash;
+
+		FWidgetHash()
+		{}
+
+		FWidgetHash(FXxHash64 InAlignmentHash, FXxHash64 InDataHash)
+			: AlignmentHash(InAlignmentHash)
+			, DataHash(InDataHash)
+		{}
+
+		bool IsValid() const
+		{
+			return !AlignmentHash.IsZero() && !DataHash.IsZero();
+		}
+	};
+	
+	struct FRootNode
+	{
+		FRootNode(FName InRootName, TSharedPtr<ISlateIMRoot> InRootWidget, bool InRootState);
+
+		FWidgetHash GetWidgetHash() const;
+
+		void SetDataHash(FXxHash64 InDataHash);
+
+		void SetAlignmentHash(FXxHash64 AlignmentHash);
+
+		void SetDisabledState() { bCurrentEnabledState = false; }
+		void SetEnabledState() { bCurrentEnabledState = true; }
+
+		void SetNextToolTip(const FStringView& InNextToolTip);
+
+		FContainerNode RootContainer;
+
+		TArray<FContainerNode> CurrentContainerStack;
+		TArray<TSharedPtr<SImContextMenuAnchor>> CurrentMenuAnchorStack;
+
+		TArray<FWidgetHash> WidgetTreeDataHash;
+
+		TSharedPtr<ISlateIMRoot> RootWidget;
+
+		FString CurrentToolTip;
+		uint64 DataHash = 0;
+		FName RootName = NAME_None;
+		int32 CurrentWidgetIndex = 0;
+		double LastAccessTime = 0;
+		bool RootState = false;
+		bool bCurrentEnabledState = true;
+		bool bActivatedThisFrame = false;
+	};
+
+	struct FWidgetActivation
+	{
+		FName RootName = NAME_None;
+		int32 ContainerIndex = INDEX_NONE;
+		int32 WidgetIndex = INDEX_NONE;
+
+		bool operator==(const FWidgetActivation& Other) const
+		{
+			return RootName == Other.RootName && ContainerIndex == Other.ContainerIndex && WidgetIndex == Other.WidgetIndex;
+		}
+	};
+	
+	class FSlateIMManager
+	{
+	public:
+		static void Initialize();
+		static FSlateIMManager& Get();
+
+		FSlateIMManager();
+		~FSlateIMManager();
+
+		template<typename RootType>
+		FRootNode* FindRoot(const FName WindowId);
+		FRootNode& AddRoot(const FName WindowId, TSharedPtr<ISlateIMRoot> NewRoot);
+
+		void BeginRoot(FName RootName);
+		void EndRoot();
+		
+		void PushContainer(FContainerNode&& Node);
+		template<typename ContainerType>
+		void PopContainer();
+		template<typename ContainerType>
+		TSharedPtr<ContainerType> GetCurrentContainer() const;
+		const FContainerNode* GetCurrentContainerNode() const;
+		
+		void PushMenuRoot(TSharedPtr<SImContextMenuAnchor>& MenuRoot);
+		void PopMenuRoot();
+		TSharedPtr<SImContextMenuAnchor> GetCurrentMenuRoot() const;
+
+		template<typename WidgetType>
+		TSharedPtr<WidgetType> BeginIMWidget();
+		TSharedPtr<SWidget> BeginCustomWidget(TSharedPtr<SWidget> ExpectedWidget);
+		void EndWidget(bool bResetAlignmentData);
+
+		void EnsureCurrentContainerNode();
+
+		/** Gets the widget which should be at the current spot in the hierarchy or null if it doesn't exist */
+		FSlateIMChild GetCurrentChild() const;
+		const FRootNode& GetCurrentRoot() const { return *CurrentRoot; } 
+		FRootNode& GetMutableCurrentRoot() { return *CurrentRoot; }
+
+		void UpdateCurrentChild(FSlateIMChild Child, const FSlateIMSlotData& AlignmentData);
+		void AdvanceToNextWidget();
+
+		void ActivateWidget(const FWidgetActivation& ActivationData);
+		bool IsWidgetActivatedThisFrame() const;
+		FWidgetActivation GetWidgetActivationData() const;
+
+		FSlateIMSlotData GetCurrentAlignmentData(
+			const FMargin& DefaultPadding,
+			const EHorizontalAlignment DefaultHAlign,
+			const EVerticalAlignment DefaultVAlign,
+			const bool bDefaultAutoSize,
+			const float DefaultMinWidth,
+			const float DefaultMinHeight,
+			const float DefaultMaxWidth,
+			const float DefaultMaxHeight) const
+		{
+			return FSlateIMSlotData(
+				NextPadding.Get(DefaultPadding),
+				NextHAlign.Get(DefaultHAlign),
+				NextVAlign.Get(DefaultVAlign),
+				NextAutoSize.Get(bDefaultAutoSize),
+				NextMinWidth.Get(DefaultMinWidth),
+				NextMinHeight.Get(DefaultMinHeight),
+				NextMaxWidth.Get(DefaultMaxWidth),
+				NextMaxHeight.Get(DefaultMaxHeight)
+			);
+		}
+		void ResetAlignmentData();
+
+		void OnSlateIMModalOpened();
+		void OnSlateIMModalClosed();
+		bool CanUpdateSlateIM() const;
+		
+		TOptional<FMargin> NextPadding;
+		TOptional<EHorizontalAlignment> NextHAlign;
+		TOptional<EVerticalAlignment> NextVAlign;
+		TOptional<bool> NextAutoSize;
+		TOptional<float> NextMinWidth;
+		TOptional<float> NextMinHeight;
+		TOptional<float> NextMaxWidth;
+		TOptional<float> NextMaxHeight;
+
+	private:
+		void ValidateRootName(FName RootName) const;
+
+		void Tick(float DeltaTime);
+
+		void OnSlateShutdown();
+
+		void RemoveUnusedChildren(FContainerNode& Container);
+
+		void PopContainer_Internal();
+
+		static TUniquePtr<FSlateIMManager> Instance;
+		
+		TMap<FName, FRootNode> RootMap;
+		TArray<FWidgetActivation, TInlineAllocator<4>> ActivatedWidgets;
+		FRootNode* CurrentRoot = nullptr;
+
+		TUniquePtr<FMemMark> MemMark;
+
+		bool bBuildingImmediateModeWidgets = false;
+		bool bIsSlateIMModalOpen = false;
+
+#if WITH_SLATEIM_EXAMPLES
+		FSlateIMTestWindowWidget TestWindowWidget;
+#if WITH_ENGINE
+		FSlateIMTestViewportWidget TestViewportWidget;
+#endif // WITH_ENGINE
+#endif // WITH_SLATEIM_EXAMPLES
+	};
+
+	template <typename RootType>
+	FRootNode* FSlateIMManager::FindRoot(const FName WindowId)
+	{
+		static_assert(std::is_base_of_v<ISlateIMRoot, RootType>, "RootType must be a subclass of ISlateIMRoot.");
+		FRootNode* RootNode = RootMap.Find(WindowId);
+
+		if (RootNode && RootNode->RootWidget && RootNode->RootWidget->IsA<RootType>())
+		{
+			return RootNode;
+		}
+
+		return nullptr;
+	}
+
+	template <typename ContainerType>
+	void FSlateIMManager::PopContainer()
+	{
+		checkf(GetCurrentContainer<ContainerType>().IsValid(), TEXT("The container being popped is not the expected type. Are your Begin and End function calls mismatched?"));
+		PopContainer_Internal();
+	}
+
+	template <typename ContainerType>
+	TSharedPtr<ContainerType> FSlateIMManager::GetCurrentContainer() const
+	{
+		static_assert(std::is_base_of_v<ISlateIMContainer, ContainerType>, "ContainerType must derive from ISlateIMContainer");
+		TSharedPtr<ISlateIMContainer> Container = CurrentRoot->CurrentContainerStack.IsEmpty() ? nullptr : CurrentRoot->CurrentContainerStack.Top().Widget;
+		if (Container && Container->IsA<ContainerType>())
+		{
+			return StaticCastSharedPtr<ContainerType>(Container);
+		}
+
+		return nullptr;
+	}
+
+	template <typename WidgetType>
+	TSharedPtr<WidgetType> FSlateIMManager::BeginIMWidget()
+	{
+		bBuildingImmediateModeWidgets = true;
+
+		EnsureCurrentContainerNode();
+		AdvanceToNextWidget();
+
+		if constexpr (std::is_base_of_v<SWidget, WidgetType>)
+		{
+			return GetCurrentChild().GetWidget<WidgetType>();
+		}
+		else if constexpr (std::is_base_of_v<ISlateIMChild, WidgetType>)
+		{
+			return GetCurrentChild().GetChild<WidgetType>();
+		}
+		else
+		{
+			checkNoEntry();
+			return nullptr;
+		}
+	}
+}
diff --git a/Engine/Plugins/Experimental/SlateIM/Source/SlateIM/Private/Misc/SlateIMSlotData.h b/Engine/Plugins/Experimental/SlateIM/Source/SlateIM/Private/Misc/SlateIMSlotData.h
new file mode 100644
--- /dev/null
+++ b/Engine/Plugins/Experimental/SlateIM/Source/SlateIM/Private/Misc/SlateIMSlotData.h
@@ -0,0 +1,67 @@
+// Copyright Epic Games, Inc. All Rights Reserved.
+
+#pragma once
+
+#include "Layout/Margin.h"
+#include "Hash/xxhash.h"
+
+namespace SlateIM
+{
+	namespace Defaults
+	{
+		const FMargin Padding(3.0f, 3.0f, 3.0f, 3.0f);
+		constexpr EHorizontalAlignment HAlign = HAlign_Left;
+		constexpr EVerticalAlignment VAlign = VAlign_Fill;
+		constexpr bool bAutoSize = true;
+		constexpr float MinWidth = 0.0f;
+		constexpr float MinHeight = 0.0f;
+		constexpr float MaxWidth = 0.0f;
+		constexpr float MaxHeight = 0.0f;
+		constexpr float InputWidgetWidth = 200.f;
+	}
+}
+
+struct FSlateIMSlotData
+{
+	FSlateIMSlotData(const FMargin& InPadding, const EHorizontalAlignment InHAlign, const EVerticalAlignment InVAlign, bool bInAutoSize, float InMinWidth, float InMinHeight, float InMaxWidth, float InMaxHeight)
+		: Padding(InPadding)
+		, HorizontalAlignment(InHAlign)
+		, VerticalAlignment(InVAlign)
+		, bAutoSize(bInAutoSize)
+		, MinWidth(InMinWidth)
+		, MinHeight(InMinHeight)
+		, MaxWidth(InMaxWidth)
+		, MaxHeight(InMaxHeight)
+		, Hash(GetAlignmentHash())
+	{
+	}
+
+	const FMargin Padding = SlateIM::Defaults::Padding;
+	const EHorizontalAlignment HorizontalAlignment = SlateIM::Defaults::HAlign;
+	const EVerticalAlignment VerticalAlignment = SlateIM::Defaults::VAlign;
+	const bool bAutoSize = SlateIM::Defaults::bAutoSize;
+	const float MinWidth = SlateIM::Defaults::MinWidth;
+	const float MinHeight = SlateIM::Defaults::MinHeight;
+	const float MaxWidth = SlateIM::Defaults::MaxWidth;
+	const float MaxHeight = SlateIM::Defaults::MaxHeight;
+	const FXxHash64 Hash;
+
+private:
+	FXxHash64 GetAlignmentHash() const
+	{
+		FXxHash64Builder HashBuilder;
+
+		HashBuilder.Update(&Padding, sizeof(FMargin));
+
+		const char AlignmentVals[] = { (char)HorizontalAlignment, (char)VerticalAlignment };
+
+		HashBuilder.Update(AlignmentVals, sizeof(AlignmentVals));
+		HashBuilder.Update(&bAutoSize, sizeof(bAutoSize));
+		HashBuilder.Update(&MinWidth, sizeof(MinWidth));
+		HashBuilder.Update(&MinHeight, sizeof(MinHeight));
+		HashBuilder.Update(&MaxWidth, sizeof(MaxWidth));
+		HashBuilder.Update(&MaxHeight, sizeof(MaxHeight));
+
+		return HashBuilder.Finalize();
+	}
+};
\ No newline at end of file
diff --git a/Engine/Plugins/Experimental/SlateIM/Source/SlateIM/Private/Misc/SlateIMTypeChecking.h b/Engine/Plugins/Experimental/SlateIM/Source/SlateIM/Private/Misc/SlateIMTypeChecking.h
new file mode 100644
--- /dev/null
+++ b/Engine/Plugins/Experimental/SlateIM/Source/SlateIM/Private/Misc/SlateIMTypeChecking.h
@@ -0,0 +1,29 @@
+﻿// Copyright Epic Games, Inc. All Rights Reserved.
+
+#pragma once
+
+class FName;
+
+class ISlateIMTypeChecking
+{
+public:
+	template<class TType>
+	bool IsA() const
+	{
+		return IsAImpl(TType::GetTypeId());
+	}
+
+	virtual ~ISlateIMTypeChecking() {}
+
+protected:
+	virtual bool IsAImpl(const FName&) const
+	{
+		return false;
+	}
+};
+
+#define SLATE_IM_TYPE_DATA(TYPE, BASE) \
+	public: \
+	static const FName& GetTypeId() { static FName Type(TEXT(#TYPE)); return Type; } \
+	private: \
+	virtual bool IsAImpl(const FName& Type) const override { return GetTypeId() == Type || BASE::IsAImpl(Type); }
diff --git a/Engine/Plugins/Experimental/SlateIM/Source/SlateIM/Private/Misc/SlateIMWidgetScope.h b/Engine/Plugins/Experimental/SlateIM/Source/SlateIM/Private/Misc/SlateIMWidgetScope.h
new file mode 100644
--- /dev/null
+++ b/Engine/Plugins/Experimental/SlateIM/Source/SlateIM/Private/Misc/SlateIMWidgetScope.h
@@ -0,0 +1,130 @@
+﻿// Copyright Epic Games, Inc. All Rights Reserved.
+
+#pragma once
+
+#include "SlateIMManager.h"
+	
+namespace SlateIM
+{
+	template<typename WidgetType>
+	struct FWidgetScope
+	{
+		FWidgetScope(const FMargin DefaultPadding = Defaults::Padding,
+			const EHorizontalAlignment DefaultHAlign = Defaults::HAlign,
+			const EVerticalAlignment DefaultVAlign = Defaults::VAlign,
+			const bool bDefaultAutoSize = Defaults::bAutoSize,
+			const float DefaultMinWidth = Defaults::MinWidth,
+			const float DefaultMinHeight = Defaults::MinHeight,
+			const float DefaultMaxWidth = Defaults::MaxWidth,
+			const float DefaultMaxHeight = Defaults::MaxHeight)
+			: Widget(FSlateIMManager::Get().BeginIMWidget<WidgetType>())
+			, AlignmentData(FSlateIMManager::Get().GetCurrentAlignmentData(DefaultPadding, DefaultHAlign, DefaultVAlign, bDefaultAutoSize, DefaultMinWidth, DefaultMinHeight, DefaultMaxWidth, DefaultMaxHeight))
+			, DataHash()
+			, bIsSlotDirty(false)
+		{
+			HashStringView(FSlateIMManager::Get().GetCurrentRoot().CurrentToolTip);
+		}
+
+		FWidgetScope(const bool bAutoSize)
+			: FWidgetScope(Defaults::Padding, Defaults::HAlign, Defaults::VAlign, bAutoSize)
+		{
+		}
+
+		FWidgetScope(const TSharedPtr<SWidget> ExpectedWidget,
+			const FMargin DefaultPadding = Defaults::Padding,
+			const EHorizontalAlignment DefaultHAlign = Defaults::HAlign,
+			const EVerticalAlignment DefaultVAlign = Defaults::VAlign,
+			const bool bDefaultAutoSize = Defaults::bAutoSize,
+			const float DefaultMinWidth = Defaults::MinWidth,
+			const float DefaultMinHeight = Defaults::MinHeight,
+			const float DefaultMaxWidth = Defaults::MaxWidth,
+			const float DefaultMaxHeight = Defaults::MaxHeight)
+			: Widget(FSlateIMManager::Get().BeginCustomWidget(ExpectedWidget))
+			, AlignmentData(FSlateIMManager::Get().GetCurrentAlignmentData(DefaultPadding, DefaultHAlign, DefaultVAlign, bDefaultAutoSize, DefaultMinWidth, DefaultMinHeight, DefaultMaxWidth, DefaultMaxHeight))
+			, DataHash()
+			, bIsSlotDirty(false)
+		{
+			HashStringView(FSlateIMManager::Get().GetCurrentRoot().CurrentToolTip);
+		}
+
+		~FWidgetScope()
+		{
+			FRootNode& CurrentRoot = FSlateIMManager::Get().GetMutableCurrentRoot();
+
+			FWidgetHash CurrentHash = CurrentRoot.GetWidgetHash();
+
+			bIsSlotDirty |= (CurrentHash.AlignmentHash != AlignmentData.Hash);
+
+			if (bIsSlotDirty)
+			{
+				FSlateIMManager::Get().UpdateCurrentChild(Widget ? FSlateIMChild(Widget.ToSharedRef()) : FSlateIMChild(), AlignmentData);
+			}
+
+			if constexpr (std::is_base_of_v<SWidget, WidgetType>)
+			{
+				if (Widget)
+				{
+					Widget->SetEnabled(CurrentRoot.bCurrentEnabledState);
+
+					DataHash = HashBuilder.Finalize();
+
+					if (CurrentHash.DataHash != DataHash)
+					{
+						Widget->SetToolTipText(FText::FromStringView(CurrentRoot.CurrentToolTip));
+						CurrentRoot.SetDataHash(DataHash);
+					}
+
+					CurrentRoot.SetNextToolTip(FStringView());
+				}
+				
+				constexpr bool bResetAlignmentData = true;
+				FSlateIMManager::Get().EndWidget(bResetAlignmentData);
+			}
+			else
+			{
+				DataHash = HashBuilder.Finalize();
+
+				if (CurrentHash.DataHash != DataHash)
+				{
+					CurrentRoot.SetDataHash(DataHash);
+				}
+
+				// This "widget" is virtual, let the next widget use the current alignment data
+				constexpr bool bResetAlignmentData = false;
+				FSlateIMManager::Get().EndWidget(bResetAlignmentData);
+			}
+		}
+
+		TSharedPtr<WidgetType> GetWidget() { return Widget; }
+
+		void UpdateWidget(TSharedPtr<WidgetType> NewWidget)
+		{
+			Widget = NewWidget;
+			bIsSlotDirty = true;
+		}
+
+		template<typename T>
+		void HashData(const T& Data)
+		{
+			HashBuilder.Update(reinterpret_cast<const char*>(&Data), sizeof(T));
+		}
+
+		void HashStringView(const FStringView& String)
+		{
+			HashBuilder.Update(reinterpret_cast<const char*>(String.GetData()), String.Len() * sizeof(TCHAR));
+		}
+
+		bool IsDataHashDirty() const
+		{
+			DataHash = HashBuilder.Finalize();
+			return FSlateIMManager::Get().GetCurrentRoot().GetWidgetHash().DataHash != DataHash;
+		}
+
+	private:
+		TSharedPtr<WidgetType> Widget;
+		FSlateIMSlotData AlignmentData;
+		FXxHash64Builder HashBuilder;
+		mutable FXxHash64 DataHash;
+		bool bIsSlotDirty = false;
+	};
+}
diff --git a/Engine/Plugins/Experimental/SlateIM/Source/SlateIM/Private/Roots/ISlateIMRoot.h b/Engine/Plugins/Experimental/SlateIM/Source/SlateIM/Private/Roots/ISlateIMRoot.h
new file mode 100644
--- /dev/null
+++ b/Engine/Plugins/Experimental/SlateIM/Source/SlateIM/Private/Roots/ISlateIMRoot.h
@@ -0,0 +1,19 @@
+// Copyright Epic Games, Inc. All Rights Reserved.
+
+#pragma once
+
+#include "Containers/StringFwd.h"
+#include "Math/MathFwd.h"
+#include "Misc/SlateIMTypeChecking.h"
+#include "Templates/SharedPointer.h"
+
+struct FSlateIMSlotData;
+class SWidget;
+
+class ISlateIMRoot : public ISlateIMTypeChecking
+{
+public:
+	virtual void UpdateRoot(const FStringView& Title, FVector2f Offset, FVector2f Size) = 0;
+	virtual void UpdateChild(TSharedRef<SWidget> Child, const FSlateIMSlotData& AlignmentData) = 0;
+	virtual bool IsVisible() const = 0;
+};
diff --git a/Engine/Plugins/Experimental/SlateIM/Source/SlateIM/Private/Roots/SlateIMExposedRoot.cpp b/Engine/Plugins/Experimental/SlateIM/Source/SlateIM/Private/Roots/SlateIMExposedRoot.cpp
new file mode 100644
--- /dev/null
+++ b/Engine/Plugins/Experimental/SlateIM/Source/SlateIM/Private/Roots/SlateIMExposedRoot.cpp
@@ -0,0 +1,40 @@
+﻿// Copyright Epic Games, Inc. All Rights Reserved.
+
+#include "SlateIMExposedRoot.h"
+
+#include "Misc/SlateIMSlotData.h"
+#include "Widgets/Layout/SBox.h"
+
+
+FSlateIMExposedRoot::FSlateIMExposedRoot()
+{
+}
+
+FSlateIMExposedRoot::~FSlateIMExposedRoot()
+{
+	if (ExposedWidget.IsValid())
+	{
+		ExposedWidget->SetContent(SNullWidget::NullWidget);
+	}
+}
+
+void FSlateIMExposedRoot::UpdateRoot(const FStringView& Title, FVector2f Offset, FVector2f Size)
+{
+	if (!ExposedWidget.IsValid())
+	{
+		SAssignNew(ExposedWidget, SBox);
+	}
+}
+
+void FSlateIMExposedRoot::UpdateChild(TSharedRef<SWidget> Child, const FSlateIMSlotData& AlignmentData)
+{
+	ExposedWidget->SetContent(Child);
+	ExposedWidget->SetPadding(AlignmentData.Padding);
+	ExposedWidget->SetHAlign(AlignmentData.HorizontalAlignment);
+	ExposedWidget->SetVAlign(AlignmentData.VerticalAlignment);
+}
+
+TSharedRef<SWidget> FSlateIMExposedRoot::GetExposedWidget() const
+{
+	return ExposedWidget ? ExposedWidget.ToSharedRef() : SNullWidget::NullWidget;
+}
diff --git a/Engine/Plugins/Experimental/SlateIM/Source/SlateIM/Private/Roots/SlateIMExposedRoot.h b/Engine/Plugins/Experimental/SlateIM/Source/SlateIM/Private/Roots/SlateIMExposedRoot.h
new file mode 100644
--- /dev/null
+++ b/Engine/Plugins/Experimental/SlateIM/Source/SlateIM/Private/Roots/SlateIMExposedRoot.h
@@ -0,0 +1,26 @@
+﻿// Copyright Epic Games, Inc. All Rights Reserved.
+
+#pragma once
+
+#include "ISlateIMRoot.h"
+
+class SBox;
+
+class FSlateIMExposedRoot : public ISlateIMRoot
+{
+	SLATE_IM_TYPE_DATA(FSlateIMExposedRoot, ISlateIMRoot)
+
+public:
+	FSlateIMExposedRoot();
+	virtual ~FSlateIMExposedRoot() override;
+
+	virtual void UpdateRoot(const FStringView& Title, FVector2f Offset, FVector2f Size) override;
+	virtual void UpdateChild(TSharedRef<SWidget> Child, const FSlateIMSlotData& AlignmentData) override;
+	virtual bool IsVisible() const override { return true; }
+
+	TSharedRef<SWidget> GetExposedWidget() const;
+	
+private:
+	TSharedPtr<SBox> ExposedWidget;
+
+};
diff --git a/Engine/Plugins/Experimental/SlateIM/Source/SlateIM/Private/Roots/SlateIMViewportRoot.cpp b/Engine/Plugins/Experimental/SlateIM/Source/SlateIM/Private/Roots/SlateIMViewportRoot.cpp
new file mode 100644
--- /dev/null
+++ b/Engine/Plugins/Experimental/SlateIM/Source/SlateIM/Private/Roots/SlateIMViewportRoot.cpp
@@ -0,0 +1,110 @@
+﻿// Copyright Epic Games, Inc. All Rights Reserved.
+
+
+#include "SlateIMViewportRoot.h"
+
+#if WITH_EDITOR
+#include "IAssetViewport.h"
+#endif
+#if WITH_ENGINE
+#include "Engine/GameViewportClient.h"
+#include "Engine/LocalPlayer.h"
+#endif
+
+
+#if WITH_ENGINE
+FSlateIMViewportRoot::FSlateIMViewportRoot(UGameViewportClient* InGameViewport)
+	: GameViewport(InGameViewport)
+{
+	ensure(InGameViewport);
+}
+
+FSlateIMViewportRoot::FSlateIMViewportRoot(ULocalPlayer* InLocalPlayer)
+	: GameViewport(InLocalPlayer ? InLocalPlayer->ViewportClient : nullptr)
+	, LocalPlayer(InLocalPlayer)
+{
+	ensure(InLocalPlayer);
+	ensure(GameViewport.IsValid());
+}
+
+#if WITH_EDITOR
+FSlateIMViewportRoot::FSlateIMViewportRoot(TSharedPtr<IAssetViewport> InAssetViewport)
+	: AssetViewport(InAssetViewport)
+{
+	ensure(InAssetViewport);
+}
+#endif // WITH_EDITOR
+
+FSlateIMViewportRoot::~FSlateIMViewportRoot()
+{
+	if (ViewportCanvas)
+	{
+#if WITH_EDITOR
+		if (TSharedPtr<IAssetViewport> AssetViewportPinned = AssetViewport.Pin())
+		{
+			AssetViewportPinned->RemoveOverlayWidget(ViewportCanvas.ToSharedRef());
+		}
+		else
+#endif // WITH_EDITOR
+		if (GameViewport.IsValid())
+		{
+			if (LocalPlayer.IsValid())
+			{
+				GameViewport->RemoveViewportWidgetForPlayer(LocalPlayer.Get(), ViewportCanvas.ToSharedRef());
+			}
+			else
+			{
+				GameViewport->RemoveGameLayerWidget(ViewportCanvas.ToSharedRef());
+			}
+		}
+	}
+}
+
+void FSlateIMViewportRoot::UpdateRoot(const FStringView& Title, FVector2f Offset, FVector2f Size)
+{
+	if (!ViewportCanvas)
+	{
+		ViewportCanvas =
+			SNew(SConstraintCanvas)
+			+ SConstraintCanvas::Slot()
+			.Offset(Offset)
+			.AutoSize(true)
+			.Anchors(FAnchors(0.f, 0.f, Size.X, Size.Y))
+			.Alignment(FVector2D(.0f, .0f))
+			.Expose(Slot);
+
+#if WITH_EDITOR
+		if (TSharedPtr<IAssetViewport> AssetViewportPinned = AssetViewport.Pin())
+		{
+			AssetViewportPinned->AddOverlayWidget(ViewportCanvas.ToSharedRef());
+		}
+		else
+#endif // WITH_EDITOR
+		if (GameViewport.IsValid())
+		{
+			int ZOrder = 10;
+			if (LocalPlayer.IsValid())
+			{
+				GameViewport->AddViewportWidgetForPlayer(LocalPlayer.Get(), ViewportCanvas.ToSharedRef(), ZOrder);
+			}
+			else
+			{
+				GameViewport->AddGameLayerWidget(ViewportCanvas.ToSharedRef(), ZOrder);
+			}
+		}
+	}
+	else
+	{
+		Slot->SetOffset(Offset);
+		Slot->SetAnchors(FAnchors(0.f, 0.f, Size.X, Size.Y));
+	}
+}
+
+void FSlateIMViewportRoot::UpdateChild(TSharedRef<SWidget> Child, const FSlateIMSlotData& AlignmentData)
+{
+	if (Slot)
+	{
+		Slot->AttachWidget(Child);
+	}
+}
+#endif // WITH_ENGINE
diff --git a/Engine/Plugins/Experimental/SlateIM/Source/SlateIM/Private/Roots/SlateIMViewportRoot.h b/Engine/Plugins/Experimental/SlateIM/Source/SlateIM/Private/Roots/SlateIMViewportRoot.h
new file mode 100644
--- /dev/null
+++ b/Engine/Plugins/Experimental/SlateIM/Source/SlateIM/Private/Roots/SlateIMViewportRoot.h
@@ -0,0 +1,46 @@
+﻿// Copyright Epic Games, Inc. All Rights Reserved.
+
+#pragma once
+
+#if WITH_ENGINE
+#include "ISlateIMRoot.h"
+#include "UObject/WeakObjectPtrTemplates.h"
+#include "Widgets/Layout/SConstraintCanvas.h"
+
+class UGameViewportClient;
+class ULocalPlayer;
+
+#if WITH_EDITOR
+class IAssetViewport;
+#endif
+
+class FSlateIMViewportRoot : public ISlateIMRoot
+{
+	SLATE_IM_TYPE_DATA(FSlateIMViewportRoot, ISlateIMRoot)
+
+public:
+	FSlateIMViewportRoot(UGameViewportClient* InGameViewport);
+	FSlateIMViewportRoot(ULocalPlayer* InLocalPlayer);
+#if WITH_EDITOR
+	FSlateIMViewportRoot(TSharedPtr<IAssetViewport> InLevelViewport);
+#endif
+	virtual ~FSlateIMViewportRoot() override;
+
+	virtual void UpdateRoot(const FStringView& Title, FVector2f Offset, FVector2f Size) override;
+	virtual void UpdateChild(TSharedRef<SWidget> Child, const FSlateIMSlotData& AlignmentData) override;
+	virtual bool IsVisible() const override
+	{
+		return true;
+	}
+
+	TWeakObjectPtr<UGameViewportClient> GameViewport;
+	TWeakObjectPtr<ULocalPlayer> LocalPlayer;
+#if WITH_EDITOR
+	TWeakPtr<IAssetViewport> AssetViewport;
+#endif
+
+private:
+	TSharedPtr<SConstraintCanvas> ViewportCanvas;
+	SConstraintCanvas::FSlot* Slot = nullptr;
+};
+#endif // WITH_ENGINE
diff --git a/Engine/Plugins/Experimental/SlateIM/Source/SlateIM/Private/Roots/SlateIMWindowRoot.cpp b/Engine/Plugins/Experimental/SlateIM/Source/SlateIM/Private/Roots/SlateIMWindowRoot.cpp
new file mode 100644
--- /dev/null
+++ b/Engine/Plugins/Experimental/SlateIM/Source/SlateIM/Private/Roots/SlateIMWindowRoot.cpp
@@ -0,0 +1,43 @@
+﻿// Copyright Epic Games, Inc. All Rights Reserved.
+
+
+#include "SlateIMWindowRoot.h"
+
+#include "Framework/Application/SlateApplication.h"
+
+FSlateIMWindowRoot::FSlateIMWindowRoot(TSharedRef<SWindow> Window)
+	: RootWindow(Window)
+{}
+
+FSlateIMWindowRoot::~FSlateIMWindowRoot()
+{
+	if (RootWindow.IsValid() && FSlateApplication::IsInitialized())
+	{
+		FSlateApplication::Get().RequestDestroyWindow(RootWindow.Pin().ToSharedRef());
+		RootWindow = nullptr;
+	}
+}
+
+void FSlateIMWindowRoot::UpdateRoot(const FStringView& Title, FVector2f Offset, FVector2f Size)
+{
+	if (TSharedPtr<SWindow> RootWindowPin = RootWindow.Pin())
+	{
+		FText NewWindowTitle = FText::FromStringView(Title);
+		if (!RootWindowPin->GetTitle().IdenticalTo(NewWindowTitle))
+		{
+			RootWindowPin->SetTitle(NewWindowTitle);
+		}
+	}
+}
+
+void FSlateIMWindowRoot::UpdateChild(TSharedRef<SWidget> Child, const FSlateIMSlotData& AlignmentData)
+{
+	check(RootWindow != Child);
+
+	RootWindow.Pin()->SetContent(Child);
+}
+
+bool FSlateIMWindowRoot::IsVisible() const
+{
+	return RootWindow.IsValid();
+}
diff --git a/Engine/Plugins/Experimental/SlateIM/Source/SlateIM/Private/Roots/SlateIMWindowRoot.h b/Engine/Plugins/Experimental/SlateIM/Source/SlateIM/Private/Roots/SlateIMWindowRoot.h
new file mode 100644
--- /dev/null
+++ b/Engine/Plugins/Experimental/SlateIM/Source/SlateIM/Private/Roots/SlateIMWindowRoot.h
@@ -0,0 +1,24 @@
+﻿// Copyright Epic Games, Inc. All Rights Reserved.
+
+#pragma once
+
+#include "ISlateIMRoot.h"
+
+class SWindow;
+
+class FSlateIMWindowRoot : public ISlateIMRoot
+{
+	SLATE_IM_TYPE_DATA(FSlateIMWindowRoot, ISlateIMRoot)
+
+public:
+	FSlateIMWindowRoot(TSharedRef<SWindow> Window);
+	virtual ~FSlateIMWindowRoot() override;
+
+	virtual void UpdateRoot(const FStringView& Title, FVector2f Offset, FVector2f Size) override;
+	virtual void UpdateChild(TSharedRef<SWidget> Child, const FSlateIMSlotData& AlignmentData) override;
+	virtual bool IsVisible() const override;
+	
+private:
+	TWeakPtr<SWindow> RootWindow;
+
+};
diff --git a/Engine/Plugins/Experimental/SlateIM/Source/SlateIM/Private/SlateIM.cpp b/Engine/Plugins/Experimental/SlateIM/Source/SlateIM/Private/SlateIM.cpp
new file mode 100644
--- /dev/null
+++ b/Engine/Plugins/Experimental/SlateIM/Source/SlateIM/Private/SlateIM.cpp
@@ -0,0 +1,20 @@
+// Copyright Epic Games, Inc. All Rights Reserved.
+
+#include "SlateIM.h"
+
+#include "Misc/SlateIMLogging.h"
+#include "Misc/SlateIMManager.h"
+#include "Modules/ModuleManager.h"
+
+DEFINE_LOG_CATEGORY(LogSlateIM);
+
+void FSlateIMModule::StartupModule()
+{
+	SlateIM::FSlateIMManager::Initialize();
+}
+
+void FSlateIMModule::ShutdownModule()
+{
+}
+
+IMPLEMENT_MODULE(FSlateIMModule, SlateIM)
diff --git a/Engine/Plugins/Experimental/SlateIM/Source/SlateIM/Private/SlateIMWidgetBase.cpp b/Engine/Plugins/Experimental/SlateIM/Source/SlateIM/Private/SlateIMWidgetBase.cpp
new file mode 100644
--- /dev/null
+++ b/Engine/Plugins/Experimental/SlateIM/Source/SlateIM/Private/SlateIMWidgetBase.cpp
@@ -0,0 +1,83 @@
+﻿// Copyright Epic Games, Inc. All Rights Reserved.
+
+
+#include "SlateIMWidgetBase.h"
+
+#include "Framework/Application/SlateApplication.h"
+#include "SlateIM.h"
+
+FSlateIMWidgetBase::FSlateIMWidgetBase(const FStringView& Name)
+	: WidgetName(Name)
+{
+}
+
+FSlateIMWidgetBase::~FSlateIMWidgetBase()
+{
+	if (FSlateApplication::IsInitialized())
+	{
+		DisableWidget();
+	}
+}
+
+void FSlateIMWidgetBase::EnableWidget()
+{
+	if (!TickHandle.IsValid())
+	{
+		TickHandle = FSlateApplication::Get().OnPreTick().AddRaw(this, &FSlateIMWidgetBase::DrawWidget);
+	}
+}
+
+void FSlateIMWidgetBase::DisableWidget()
+{
+	if (TickHandle.IsValid())
+	{
+		FSlateApplication::Get().OnPreTick().Remove(TickHandle);
+		TickHandle.Reset();
+	}
+}
+
+FSlateIMWidgetWithCommandBase::FSlateIMWidgetWithCommandBase(const TCHAR* Command, const TCHAR* CommandHelp)
+	: FSlateIMWidgetBase(Command)
+	, WidgetCommand(Command, CommandHelp, FConsoleCommandDelegate::CreateRaw(this, &FSlateIMWidgetWithCommandBase::ToggleWidget))
+{}
+
+void FSlateIMWindowBase::DrawWidget(float DeltaTime)
+{
+	if (!SlateIM::CanUpdateSlateIM())
+	{
+		return;
+	}
+	
+	const bool bIsDrawingWindow = SlateIM::BeginWindowRoot(GetWidgetName(), WindowTitle, WindowSize);
+	if (bIsDrawingWindow)
+	{
+		DrawWindow(DeltaTime);
+	}
+	SlateIM::EndRoot();
+
+	if (!bIsDrawingWindow)
+	{
+		DisableWidget();
+	}
+}
+
+TSharedRef<SWidget> FSlateIMExposedBase::GetExposedWidget() const
+{
+	return ExposedWidget.IsValid() ? ExposedWidget.ToSharedRef() : SNullWidget::NullWidget;
+}
+
+void FSlateIMExposedBase::DrawWidget(float DeltaTime)
+{
+	TSharedPtr<SWidget> NewExposedWidget;
+	if (SlateIM::BeginExposedRoot(GetWidgetName(), NewExposedWidget))
+	{
+		DrawContent(DeltaTime);
+	}
+	SlateIM::EndRoot();
+
+	if (NewExposedWidget != ExposedWidget)
+	{
+		ExposedWidget = NewExposedWidget;
+		OnExposedWidgetChanged.Broadcast(GetExposedWidget());
+	}
+}
diff --git a/Engine/Plugins/Experimental/SlateIM/Source/SlateIM/Private/Widgets/SImButton.cpp b/Engine/Plugins/Experimental/SlateIM/Source/SlateIM/Private/Widgets/SImButton.cpp
new file mode 100644
--- /dev/null
+++ b/Engine/Plugins/Experimental/SlateIM/Source/SlateIM/Private/Widgets/SImButton.cpp
@@ -0,0 +1,38 @@
+﻿// Copyright Epic Games, Inc. All Rights Reserved.
+
+#include "SImButton.h"
+
+#include "Widgets/Text/STextBlock.h"
+
+SLATE_IMPLEMENT_WIDGET(SImButton)
+
+void SImButton::PrivateRegisterAttributes(FSlateAttributeInitializer& AttributeInitializer)
+{
+}
+
+void SImButton::Construct(const FArguments& InArgs)
+{
+	SButton::FArguments ButtonArgs;
+
+	TextBox = SNew(STextBlock)
+		.Visibility(EVisibility::HitTestInvisible)
+		.TextStyle(ButtonArgs._TextStyle)
+		.TextShapingMethod(ButtonArgs._TextShapingMethod)
+		.TextFlowDirection(ButtonArgs._TextFlowDirection);
+
+	ButtonArgs.OnClicked(InArgs._OnClicked);
+
+	ButtonArgs
+	.HAlign(HAlign_Center)
+	.VAlign(VAlign_Center)
+	[
+		TextBox.ToSharedRef()
+	];
+
+	SButton::Construct(ButtonArgs);
+}
+
+void SImButton::SetText(const FStringView& InText)
+{
+	TextBox->SetText(FText::FromStringView(InText));
+}
diff --git a/Engine/Plugins/Experimental/SlateIM/Source/SlateIM/Private/Widgets/SImButton.h b/Engine/Plugins/Experimental/SlateIM/Source/SlateIM/Private/Widgets/SImButton.h
new file mode 100644
--- /dev/null
+++ b/Engine/Plugins/Experimental/SlateIM/Source/SlateIM/Private/Widgets/SImButton.h
@@ -0,0 +1,25 @@
+﻿// Copyright Epic Games, Inc. All Rights Reserved.
+
+#pragma once
+
+#include "Widgets/Input/SButton.h"
+
+
+class STextBlock;
+
+class SImButton : public SButton
+{
+	SLATE_DECLARE_WIDGET(SImButton, SButton)
+
+public:
+	SLATE_BEGIN_ARGS(SImButton) {}
+		SLATE_EVENT(FOnClicked, OnClicked)
+	SLATE_END_ARGS()
+
+	void Construct(const FArguments& InArgs);
+
+	void SetText(const FStringView& InText);
+
+private:
+	TSharedPtr<STextBlock> TextBox;
+};
diff --git a/Engine/Plugins/Experimental/SlateIM/Source/SlateIM/Private/Widgets/SImCheckBox.cpp b/Engine/Plugins/Experimental/SlateIM/Source/SlateIM/Private/Widgets/SImCheckBox.cpp
new file mode 100644
--- /dev/null
+++ b/Engine/Plugins/Experimental/SlateIM/Source/SlateIM/Private/Widgets/SImCheckBox.cpp
@@ -0,0 +1,29 @@
+﻿// Copyright Epic Games, Inc. All Rights Reserved.
+
+#include "SImCheckBox.h"
+
+#include "Widgets/Text/STextBlock.h"
+
+SLATE_IMPLEMENT_WIDGET(SImCheckBox)
+
+void SImCheckBox::PrivateRegisterAttributes(FSlateAttributeInitializer& AttributeInitializer)
+{
+}
+
+void SImCheckBox::Construct(const FArguments& InArgs)
+{
+	TextBox = SNew(STextBlock).Visibility(EVisibility::HitTestInvisible);
+	SCheckBox::Construct(
+		SCheckBox::FArguments()
+		.IsChecked(InArgs._IsChecked)
+		.OnCheckStateChanged(InArgs._OnCheckStateChanged)
+		[
+			TextBox.ToSharedRef()
+		]
+	);
+}
+
+void SImCheckBox::SetText(const FText& InText)
+{
+	TextBox->SetText(InText);
+}
diff --git a/Engine/Plugins/Experimental/SlateIM/Source/SlateIM/Private/Widgets/SImCheckBox.h b/Engine/Plugins/Experimental/SlateIM/Source/SlateIM/Private/Widgets/SImCheckBox.h
new file mode 100644
--- /dev/null
+++ b/Engine/Plugins/Experimental/SlateIM/Source/SlateIM/Private/Widgets/SImCheckBox.h
@@ -0,0 +1,25 @@
+﻿// Copyright Epic Games, Inc. All Rights Reserved.
+
+#pragma once
+
+#include "Widgets/Input/SCheckBox.h"
+
+class STextBlock;
+
+class SImCheckBox : public SCheckBox
+{
+	SLATE_DECLARE_WIDGET(SImCheckBox, SCheckBox)
+	
+public:
+	SLATE_BEGIN_ARGS(SImCheckBox) {}
+		SLATE_ATTRIBUTE(ECheckBoxState, IsChecked)
+		SLATE_EVENT(FOnCheckStateChanged, OnCheckStateChanged)
+	SLATE_END_ARGS()
+
+	void Construct(const FArguments& InArgs);
+
+	void SetText(const FText& InText);
+
+private:
+	TSharedPtr<STextBlock> TextBox;
+};
\ No newline at end of file
diff --git a/Engine/Plugins/Experimental/SlateIM/Source/SlateIM/Public/SlateIM.h b/Engine/Plugins/Experimental/SlateIM/Source/SlateIM/Public/SlateIM.h
new file mode 100644
--- /dev/null
+++ b/Engine/Plugins/Experimental/SlateIM/Source/SlateIM/Public/SlateIM.h
@@ -0,0 +1,604 @@
+// Copyright Epic Games, Inc. All Rights Reserved.
+
+#pragma once
+
+
+#include "Modules/ModuleInterface.h"
+#include "Styling/SlateColor.h"
+#include "Types/SlateEnums.h"
+#include "Layout/Margin.h"
+
+#if WITH_EDITOR
+class IAssetViewport;
+#endif
+#if WITH_ENGINE
+class UGameViewportClient;
+class ULocalPlayer;
+class UMaterialInterface;
+class UTextureRenderTarget2D;
+class UTexture2D;
+#endif
+
+struct FButtonStyle;
+struct FCheckBoxStyle;
+struct FComboBoxStyle;
+struct FEditableTextBoxStyle;
+struct FProgressBarStyle;
+struct FSlateBrush;
+struct FSliderStyle;
+struct FSpinBoxStyle;
+struct FTableViewStyle;
+struct FTextBlockStyle;
+
+enum class ECheckBoxState : uint8;
+
+class SWidget;
+
+
+class FSlateIMModule : public IModuleInterface
+{
+public:
+	virtual void StartupModule() override;
+	virtual void ShutdownModule() override;
+};
+
+/**
+ * This is the full list of available functions for creating SlateIM widgets. See FSlateIMTestWidget::Draw() for examples.
+ */
+namespace SlateIM
+{
+#pragma region Roots
+	/**
+	 * Begins a new floating window root
+	 *
+	 * @param UniqueName a unique name to identify this window
+	 * @param WindowTitle The name of the window. This will also be displayed in the windows title bar
+	 * @param WindowSize The size to open the window at
+	 * @param bShouldReopen (optional) If this window was created with a previous call to BeginWindow and then closed, passing in true will reopen the window and false will leave it closed. If the window has never been seen this parameter does nothing and the window will open
+	 *
+	 * @return the current window state. True if open, false if closed or not updating
+	 *
+	 * @note Consider using FSlateIMWindowBase to handle creating and updating your Slate IM Window
+	 */
+	SLATEIM_API bool BeginWindowRoot(FName UniqueName, const FStringView& WindowTitle, FVector2f WindowSize, bool bShouldReopen = false);
+
+#if WITH_ENGINE
+	// TODO - Game viewport widgets might need something to control HitTestability and cursor visibility
+	/**
+	 *	Begins a new root in the game viewport
+	 * 
+	 * @param UniqueName a globally unique name to identify this root
+	 * @param ViewportClient the viewport client to add the widget to
+	 * 
+	 * @return Whether the root is valid and updating
+	 */
+	SLATEIM_API bool BeginViewportRoot(FName UniqueName, UGameViewportClient* ViewportClient);
+	
+	/**
+	 *	Begins a new root in the provided player's viewport
+	 * 
+	 * @param UniqueName a globally unique name to identify this root
+	 * @param LocalPlayer the player whose viewport to add the widget to
+	 * 
+	 * @return Whether the root is valid and updating
+	 */
+	SLATEIM_API bool BeginViewportRoot(FName UniqueName, ULocalPlayer* LocalPlayer);
+#endif
+
+#if WITH_EDITOR
+	/**
+	 * Begins a SlateIM root in the provided in editor viewport
+	 * 
+	 * @param UniqueName Unique identifier for this root. Must be globally unique within the same editor session
+	 * @param AssetViewport The editor viewpor to add the SlateIM content to
+	 * @param Offset Offset from the start top left of the viewport corner. Not normalized. 
+	 * @param NormalizedSize Size of the content relative to the size of the viewport in normalized coordinates. I.E as a percentage of viewport dimensions: 0.5 means half the viewport size, 1.0 means full viewport size
+	 * 
+	 * @return Whether the root is valid and updating
+	 */
+	SLATEIM_API bool BeginViewportRoot(FName UniqueName, TSharedPtr<IAssetViewport> AssetViewport, FVector2f Offset, FVector2f NormalizedSize);
+#endif
+
+	/**
+	 * Begins a SlateIM root that exposes a slate widget for embedding in other Slate hierarchies
+	 * 
+	 * @param UniqueName a globally unique name to identify this root
+	 * @param OutSlateIMWidget where to output the resulting SlateIM widget to
+	 * 
+	 * @return Whether the root is valid and updating
+	 * 
+	 * @note Consider using FSlateIMExposedBase to handle creating and updating your Slate IM exposed widget
+	 */
+	SLATEIM_API bool BeginExposedRoot(FName UniqueName, TSharedPtr<SWidget>& OutSlateIMWidget);
+
+	/**
+	 * Ends any Root type, must always be called regardless of the result of the Begin function
+	 */
+	SLATEIM_API void EndRoot();
+#pragma endregion Roots
+
+#pragma region Containers
+	/**
+	 * Begins a horizontally stacked container. All widgets created within the Horizontal Stack container will be placed in-order left-to-right
+	 */
+	SLATEIM_API void BeginHorizontalStack();
+	SLATEIM_API void EndHorizontalStack();
+
+	/**
+	 * Begins a vertically stacked container. All widgets created within the Vertical Stack container will be placed in-order top-to-bottom
+	 */
+	SLATEIM_API void BeginVerticalStack();
+	SLATEIM_API void EndVerticalStack();
+
+	/**
+	 * Begins a horizontally wrapped container. All widgets created within the Horizontal Wrap container will be placed in-order left-to-right until the allotted width is filled,
+	 * then the content will begin a new line and so-on.
+	 *
+	 * @note Wrap box children are always AutoSize and do not support Fill
+	 */
+	SLATEIM_API void BeginHorizontalWrap();
+	SLATEIM_API void EndHorizontalWrap();
+	
+	/**
+	 * Begins a vertically wrapped container. All widgets created within the Vertical Wrap container will be placed in-order top-to-bottom until the allotted height is filled,
+	 * then the content will begin a new column and so-on.
+	 *
+	 * @note Wrap container children are always AutoSize and do not support Fill
+	 */
+	SLATEIM_API void BeginVerticalWrap();
+	SLATEIM_API void EndVerticalWrap();
+
+	/**
+	 * Begin a Table container. The contents are automatically scrollable if the table is sized smaller than all its content.
+	 * 
+	 * @param InStyle (optional) A Table View Style to override the default style. Changes to the style after the first invocation are not reflected,
+	 *				  provide a new style object to update the visuals at runtime
+	 */
+	SLATEIM_API void BeginTable(const FTableViewStyle* InStyle = nullptr);
+	/**
+	 * End the current Table container.
+	 */
+	SLATEIM_API void EndTable();
+
+	/**
+	 * Add a column to a label. Call this before adding any content to the table.
+	 * 
+	 * @param ColumnLabel (optional) The text label to display in the header for this column
+	 *
+	 * @note Table columns support SetToolTip, HAlign, and VAlign. See FixedTableColumnWidth() and InitialTableColumnWidth() for size options.
+	 */
+	SLATEIM_API void AddTableColumn(const FStringView& ColumnLabel = FStringView());
+
+	/**
+	 * Set the column to a fixed width, not resizable by users. Call this before calling AddTableColumn().
+	 * 
+	 * @param Width The fixed width of the column
+	 */
+	SLATEIM_API void FixedTableColumnWidth(float Width);
+
+	/**
+	 * Set the column to an initial width, which can then be resized by users. Call this before calling AddTableColumn().
+	 * 
+	 * @param Width The initial width of the column
+	 */
+	SLATEIM_API void InitialTableColumnWidth(float Width);
+
+	/**
+	 * Begin the content for the next table cell. The column and row are tracked internally based on the number of columns.
+	 * 
+	 * @return Whether the contents of this table cell are visible. When false, the content can be skipped to save cycles.
+	 *
+	 * @note There is no "End" function counterpart to this call. This function supports the various slot functions (Padding, HAlign, VAlign, Min/Max Width/Height, etc).
+	 */
+	SLATEIM_API bool NextTableCell();
+
+	/**
+	 * Begin adding content as a child to the current table row. NextTableCell() must still be called after this function.
+	 * 
+	 * @return Whether the parent row is expanded or not. When false, the child content can be skipped to save cycles.
+	 */
+	SLATEIM_API bool BeginTableRowChildren();
+	/**
+	 * Stop adding child content to the parent row. This must be called no matter the result of BeginTableRowChildren()
+	 */
+	SLATEIM_API void EndTableRowChildren();
+
+	/**
+	 * Begin a container with a background image.
+	 * 
+	 * @param BorderBrush The brush to use as the background for this container
+	 * @param Orientation Which direction the contents of the container should flow
+	 * @param bAbsorbMouse Whether the container should handle all mouse inputs
+	 * @param ContentPadding How much to pad the contents of the container
+	 */
+	SLATEIM_API void BeginBorder(const FSlateBrush* BorderBrush, EOrientation Orientation = Orient_Vertical, bool bAbsorbMouse = true, FMargin ContentPadding = FMargin(2.0f));
+	
+	/**
+	 * Begin a container with a background image.
+	 * 
+	 * @param BorderStyleName The name of the brush in FAppStyle to use as the background for this container
+	 * @param Orientation Which direction the contents of the container should flow
+	 * @param bAbsorbMouse Whether the container should handle all mouse inputs
+	 * @param ContentPadding How much to pad the contents of the container
+	 */
+	SLATEIM_API void BeginBorder(const FName BorderStyleName, EOrientation Orientation = Orient_Vertical, bool bAbsorbMouse = true, FMargin ContentPadding = FMargin(2.0f));
+	SLATEIM_API void EndBorder();
+
+	/**
+	 * Begins a container that will allow the user to scroll when its content is larger than the allotted space.
+	 * 
+	 * @param ScrollBarOrientation Which direction content should flow and scroll
+	 * 
+	 * @return Whether the user scroll the content this frame
+	 */
+	SLATEIM_API bool BeginScrollBox(EOrientation ScrollBarOrientation = EOrientation::Orient_Vertical);
+	SLATEIM_API void EndScrollBox();
+#pragma endregion Containers
+
+#pragma region Slots
+	/**
+	 * Sets the padding for the next widget
+	 * 
+	 * @param NextPadding The padding to apply to the next widget
+	 */
+	SLATEIM_API void Padding(const FMargin NextPadding);
+	
+	/**
+	 * Sets the horizontal alignment for the next widget
+	 * 
+	 * @param NextAlignment The alignment to apply to the next widget
+	 */
+	SLATEIM_API void HAlign(EHorizontalAlignment NextAlignment);
+	/**
+	 * Sets the vertical alignment for the next widget
+	 * 
+	 * @param NextAlignment The alignment to apply to the next widget
+	 */
+	SLATEIM_API void VAlign(EVerticalAlignment NextAlignment);
+
+	/**
+	 * Set the next slot to AutoSize to its content
+	 */
+	SLATEIM_API void AutoSize();
+	/**
+	 * Sets the next slot to Fill the remaining space in its container
+	 *
+	 * @note The filled space is shared equally with any other slots set to Fill the same container
+	 */
+	SLATEIM_API void Fill();
+
+	/**
+	 * Set the minimum width the next slot can have.
+	 * 
+	 * @param InMinWidth The minimum width for the next slot
+	 */
+	SLATEIM_API void MinWidth(float InMinWidth);
+	/**
+	 * Set the minimum height the next slot can have.
+	 * 
+	 * @param InMinHeight The minimum height for the next slot
+	 */
+	SLATEIM_API void MinHeight(float InMinHeight);
+	/**
+	 * Set the maximum width the next slot can have.
+	 * 
+	 * @param InMaxWidth The maximum width for the next slot
+	 */
+	SLATEIM_API void MaxWidth(float InMaxWidth);
+	/**
+	 * Set the maximum height the next slot can have.
+	 * 
+	 * @param InMaxHeight The maximum height for the next slot
+	 */
+	SLATEIM_API void MaxHeight(float InMaxHeight);
+#pragma endregion Slots
+
+#pragma region Widgets
+	/**
+	 * Display a string of text
+	 * 
+	 * @param InText The text string to display
+	 * @param TextStyle (optional) A Text Block Style to override the default style. Changes to the style after the first invocation are not reflected,
+	 *					provide a new style object to update the visuals at runtime
+	 */
+	SLATEIM_API void Text(const FStringView& InText, const FTextBlockStyle* TextStyle = nullptr);
+	/**
+	 * Display a colored string of text
+	 * 
+	 * @param InText The text string to display
+	 * @param Color The color to display the text as
+	 * @param TextStyle (optional) A Text Block Style to override the default style. Changes to the style after the first invocation are not reflected,
+	 *					provide a new style object to update the visuals at runtime
+	 */
+	SLATEIM_API void Text(const FStringView& InText, FSlateColor Color, const FTextBlockStyle* TextStyle = nullptr);
+
+	/**
+	 * Create a text input field
+	 * 
+	 * @param InOutText The text that is in the input field. Can have a default value
+	 * @param HintText The hint text to display when the input field is empty
+	 * @param TextStyle (optional) An Editable Text Box Style to override the default style. Changes to the style after the first invocation are not reflected,
+	 *					provide a new style object to update the visuals at runtime
+	 *					
+	 * @return Whether the user changed or committed text this frame
+	 */
+	SLATEIM_API bool EditableText(FString& InOutText, const FStringView& HintText = FStringView(), const FEditableTextBoxStyle* TextStyle = nullptr);
+
+	/**
+	 * Display a brush
+	 * 
+	 * @param ImageBrush The brush to display
+	 * @param ColorAndOpacity (optional) The color to tint the brush
+	 * @param DesiredSize (optional) Override the desired size to display the brush at
+	 */
+	SLATEIM_API void Image(const FSlateBrush* ImageBrush, const FSlateColor& ColorAndOpacity = FLinearColor::White, FVector2D DesiredSize = FVector2D::ZeroVector);
+	/**
+	 * Display a slate style brush
+	 * 
+	 * @param ImageStyleName The named of the brush to display from FAppStyle
+	 * @param ColorAndOpacity (optional) The color to tint the brush
+	 * @param DesiredSize (optional) Override the desired size to display the brush at
+	 */
+	SLATEIM_API void Image(const FName ImageStyleName, const FSlateColor& ColorAndOpacity = FLinearColor::White, FVector2D DesiredSize = FVector2D::ZeroVector);
+#if WITH_ENGINE
+	/**
+	 * Display a texture
+	 * 
+	 * @param ImageTexture The texture to display
+	 * @param ColorAndOpacity (optional) The color to tint the brush
+	 * @param DesiredSize (optional) Override the desired size to display the brush at
+	 */
+	SLATEIM_API void Image(UTexture2D* ImageTexture, const FSlateColor& ColorAndOpacity = FLinearColor::White, FVector2D DesiredSize = FVector2D::ZeroVector);
+	/**
+	 * Display a render target texture
+	 * 
+	 * @param ImageRenderTarget The render target texture to display
+	 * @param ColorAndOpacity (optional) The color to tint the brush
+	 * @param DesiredSize (optional) Override the desired size to display the brush at
+	 */
+	SLATEIM_API void Image(UTextureRenderTarget2D* ImageRenderTarget, const FSlateColor& ColorAndOpacity = FLinearColor::White, FVector2D DesiredSize = FVector2D::ZeroVector);
+	/**
+	 * Display a material
+	 * 
+	 * @param ImageMaterial The material to display
+	 * @param BrushSize The size to create the internal brush with
+	 * @param ColorAndOpacity (optional) The color to tint the brush
+	 * @param DesiredSize (optional) Override the desired size to display the brush at
+	 */
+	SLATEIM_API void Image(UMaterialInterface* ImageMaterial, FVector2D BrushSize, const FSlateColor& ColorAndOpacity = FLinearColor::White, FVector2D DesiredSize = FVector2D::ZeroVector);
+#endif
+
+	/**
+	 * Display a button with text
+	 * 
+	 * @param InText The text to display on the button
+	 * @param InStyle (optional) A Button Style to override the default style. Changes to the style after the first invocation are not reflected,
+	 *				  provide a new style object to update the visuals at runtime
+	 *					
+	 * @return Whether the user clicked the button this frame
+	 */
+	SLATEIM_API bool Button(const FStringView& InText, const FButtonStyle* InStyle = nullptr);
+
+	/**
+	 * Display a two-state checkbox
+	 * 
+	 * @param InText The label for the checkbox
+	 * @param InOutCurrentState The current state of the checkbox. Can provide a default value.
+	 * @param CheckBoxStyle (optional) A Checkbox Style to override the default style. Changes to the style after the first invocation are not reflected,
+	 *						provide a new style object to update the visuals at runtime
+	 *					
+	 * @return Whether the user changed the state of the checkbox this frame
+	 */
+	SLATEIM_API bool CheckBox(const FStringView& InText, bool& InOutCurrentState, const FCheckBoxStyle* CheckBoxStyle = nullptr);
+	
+	/**
+	 * Display a three-state checkbox
+	 * 
+	 * @param InText The label for the checkbox
+	 * @param InOutCurrentState The current state of the checkbox. Can provide a default value.
+	 * @param CheckBoxStyle (optional) A Checkbox Style to override the default style. Changes to the style after the first invocation are not reflected,
+	 *						provide a new style object to update the visuals at runtime
+	 *					
+	 * @return Whether the user changed the state of the checkbox this frame
+	 */
+	SLATEIM_API bool CheckBox(const FStringView& InText, ECheckBoxState& InOutCurrentState, const FCheckBoxStyle* CheckBoxStyle = nullptr);
+
+	/**
+	 * Display a float-based spin box
+	 * 
+	 * @param InOutValue The current value of the spin box. Can provide a default value.
+	 * @param Min The minimum value of the spin box (or unset for no limit).
+	 * @param Max The maximum value of the spin box (or unset for no limit).
+	 * @param SpinBoxStyle (optional) A Spin Box Style to override the default style. Changes to the style after the first invocation are not reflected,
+	 *						provide a new style object to update the visuals at runtime
+	 * 
+	 * @return Whether the user changed or committed the spin box value this frame
+	 */
+	SLATEIM_API bool SpinBox(float& InOutValue, TOptional<float> Min, TOptional<float> Max, const FSpinBoxStyle* SpinBoxStyle = nullptr);
+	/**
+	 * Display a double-based spin box
+	 * 
+	 * @param InOutValue The current value of the spin box. Can provide a default value.
+	 * @param Min The minimum value of the spin box (or unset for no limit).
+	 * @param Max The maximum value of the spin box (or unset for no limit).
+	 * @param SpinBoxStyle (optional) A Spin Box Style to override the default style. Changes to the style after the first invocation are not reflected,
+	 *						provide a new style object to update the visuals at runtime
+	 * 
+	 * @return Whether the user changed or committed the spin box value this frame
+	 */
+	SLATEIM_API bool SpinBox(double& InOutValue, TOptional<double> Min, TOptional<double> Max, const FSpinBoxStyle* SpinBoxStyle = nullptr);
+	/**
+	 * Display an integer-based spin box
+	 * 
+	 * @param InOutValue The current value of the spin box. Can provide a default value.
+	 * @param Min The minimum value of the spin box (or unset for no limit).
+	 * @param Max The maximum value of the spin box (or unset for no limit).
+	 * @param SpinBoxStyle (optional) A Spin Box Style to override the default style. Changes to the style after the first invocation are not reflected,
+	 *						provide a new style object to update the visuals at runtime
+	 * 
+	 * @return Whether the user committed the spin box value this frame
+	 */
+	SLATEIM_API bool SpinBox(int32& InOutValue, TOptional<int32> Min, TOptional<int32> Max, const FSpinBoxStyle* SpinBoxStyle = nullptr);
+
+	/**
+	 * Display a float-based slider
+	 * 
+	 * @param InOutValue The current value of the slider. Can provide a default value.
+	 * @param Min The minimum value of the slider
+	 * @param Max The maximum value of the slider
+	 * @param Step The smallest incremental change that can be made to the slider's value
+	 * @param SliderStyle (optional) A Slider Style to override the default style. Changes to the style after the first invocation are not reflected,
+	 *					  provide a new style object to update the visuals at runtime
+	 *					  
+	 * @return Whether the user changed the value of the slider this frame
+	 */
+	SLATEIM_API bool Slider(float& InOutValue, float Min, float Max, float Step, const FSliderStyle* SliderStyle = nullptr);
+
+	/**
+	 * Display a progress bar
+	 * 
+	 * @param Percent The value of the progress bar (0.0 to 1.0)
+	 * @param ProgressBarStyle (optional) A Progress Bar Style to override the default style. Changes to the style after the first invocation are not reflected,
+	 *						   provide a new style object to update the visuals at runtime
+	 */
+	SLATEIM_API void ProgressBar(TOptional<float> Percent, const FProgressBarStyle* ProgressBarStyle = nullptr);
+
+	/**
+	 * Display a dropdown of text options
+	 * 
+	 * @param ComboItems The options the user can choose from
+	 * @param InOutSelectedItemIndex The index of the currently selected option
+	 * @param InComboStyle (optional) A Combo Box Style to override the default style. Changes to the style after the first invocation are not reflected,
+	 *					   provide a new style object to update the visuals at runtime
+	 *					   
+	 * @return Whether the user changed the selected option this frame
+	 */
+	SLATEIM_API bool ComboBox(const TArray<FString>& ComboItems, int32& InOutSelectedItemIndex, const FComboBoxStyle* InComboStyle = nullptr);
+
+	/**
+	 * Display a list of text options
+	 * 
+	 * @param ListItems The list of options the user can choose from
+	 * @param InOutSelectedItemIndex The index of the currently selected option
+	 * @param bForceRefresh Whether to force a refresh of the available options or the selected option (set to true for a frame when changing the list of options or manually setting the selected index)
+	 * @param InStyle (optional) A Table View Style to override the default style. Changes to the style after the first invocation are not reflected,
+	 *				  provide a new style object to update the visuals at runtime
+	 *					   
+	 * @return Whether the user changed the selected option this frame
+	 */
+	SLATEIM_API bool SelectionList(const TArray<FString>& ListItems, int32& InOutSelectedItemIndex, bool bForceRefresh, const FTableViewStyle* InStyle = nullptr);
+
+	/**
+	 * Create a block of empty space
+	 * 
+	 * @param Size The size of whitespace to create
+	 */
+	SLATEIM_API void Spacer(const FVector2D& Size);
+
+	/** 
+	 * Add an already created slate widget to the IM hierarchy
+	 * 
+	 * @note Do not create the widget you are passing in here every frame. This is meant to be used to add already created widgets
+	 */
+	SLATEIM_API void Widget(TSharedRef<SWidget> InWidget);
+#pragma endregion Widgets
+
+#pragma region Menu
+	/**
+	 * Begin an area where a menu appears with a right click
+	 *
+	 * @return Whether the menu is open
+	 */
+	SLATEIM_API bool BeginContextMenuAnchor();
+	SLATEIM_API void EndContextMenuAnchor();
+
+	/**
+	 * Add a new section to the current context menu
+	 * 
+	 * @param SectionText The label of the section to add
+	 */
+	SLATEIM_API void AddMenuSection(const FStringView& SectionText);
+	
+	/**
+	 * Add a separator to the current context menu
+	 */
+	SLATEIM_API void AddMenuSeparator();
+
+	/**
+	 * Add a menu item button to the current context menu
+	 * 
+	 * @param RowText Label for the menu item
+	 * @param ToolTipText (optional) The tooltip to display for the menu item
+	 * 
+	 * @return Whether the user clicked the menu item this frame
+	 */
+	SLATEIM_API bool AddMenuButton(const FStringView& RowText, const FStringView& ToolTipText = FStringView());
+	/**
+	 * Display a menu item button with a checkbox
+	 * 
+	 * @param RowText Label for the menu item
+	 * @param InOutCurrentState The current state of the toggle
+	 * @param ToolTipText (optional) The tooltip to display for the menu item
+	 * 
+	 * @return Whether the user toggled the checkbox this frame
+	 */
+	SLATEIM_API bool AddMenuToggleButton(const FStringView& RowText, bool& InOutCurrentState, const FStringView& ToolTipText = FStringView());
+	/**
+	 * Display a menu item button with a checkmark
+	 * 
+	 * @param RowText Label for the menu item
+	 * @param InOutCurrentState The current state of the checkmark
+	 * @param ToolTipText (optional) The tooltip to display for the menu item
+	 * 
+	 * @return Whether the user clicked the menu item this frame
+	 */
+	SLATEIM_API bool AddMenuCheckButton(const FStringView& RowText, bool& InOutCurrentState, const FStringView& ToolTipText = FStringView());
+
+	/**
+	 * Add a submenu item to the current menu
+	 * 
+	 * @param SubMenuText The label for the menu item
+	 */
+	SLATEIM_API void BeginSubMenu(const FStringView& SubMenuText);
+	SLATEIM_API void EndSubMenu();
+#pragma endregion Menu
+
+#pragma region Util
+	/**
+	 * SlateIM updates are disabled in specific scenarios (like when a SlateIM::ModalDialog is open), use this function to react accordingly
+	 * 
+	 * @return Whether SlateIM can update currently
+	 */
+	SLATEIM_API bool CanUpdateSlateIM();
+	
+	/**
+	 * Disables all widgets until EndDisabledState is called.
+	 * 
+	 * @note It is not possible to enable a widget inside a disabled parent widget by calling EndDisabledState before a child is created inside a disabled Widget.
+	 */
+	SLATEIM_API void BeginDisabledState();
+	SLATEIM_API void EndDisabledState();
+
+	/**
+	 * Sets the tooltip to be used for the next widget created. Resets after tooltip is used
+	 *
+	 * @param NextToolTip The tooltip to display
+	 */
+	SLATEIM_API void SetToolTip(const FStringView& NextToolTip);
+
+	/**
+	 * Opens a modal dialog of the specified type. This function will not return until the user closes the modal dialog
+	 * 
+	 * @param MessageType The type of options the user can respond to the dialog with
+	 * @param DialogText The message to display to the user
+	 * @param Category The type of dialog to display
+	 * @param DialogTitle The title to display in the dialog window
+	 * 
+	 * @return The option the user selected
+	 *
+	 * @note SlateIM updates are disable while a SlateIM modal is open
+	 * @see CanUpdateSlateIM()
+	 */
+	SLATEIM_API EAppReturnType::Type ModalDialog(EAppMsgType::Type MessageType, const FStringView& DialogText, EAppMsgCategory Category = EAppMsgCategory::Warning, const FStringView& DialogTitle = FStringView());
+#pragma endregion Util
+}
diff --git a/Engine/Plugins/Experimental/SlateIM/Source/SlateIM/Public/SlateIMWidgetBase.h b/Engine/Plugins/Experimental/SlateIM/Source/SlateIM/Public/SlateIMWidgetBase.h
new file mode 100644
--- /dev/null
+++ b/Engine/Plugins/Experimental/SlateIM/Source/SlateIM/Public/SlateIMWidgetBase.h
@@ -0,0 +1,102 @@
+﻿// Copyright Epic Games, Inc. All Rights Reserved.
+
+#pragma once
+
+#include "Delegates/IDelegateInstance.h"
+#include "HAL/IConsoleManager.h"
+#include "Math/Vector2D.h"
+
+class SWidget;
+
+/**
+ * A utility base class for creating SlateIM-based tools.
+ * Handles setting up a tick for drawing the tool.
+ * 
+ * Note: this class is not required for using SlateIM but provides common functionality for SlateIM tools
+ *
+ * The test widgets (FSlateIMTestViewportWidget/FSlateIMTestWindowWidget in SlateIMManager) are examples of using this class.
+ */
+class SLATEIM_API FSlateIMWidgetBase
+{
+public:
+	explicit FSlateIMWidgetBase(const FStringView& Name);
+	virtual ~FSlateIMWidgetBase();
+
+	void ToggleWidget() { IsWidgetEnabled() ? DisableWidget() : EnableWidget(); }
+	void EnableWidget();
+	void DisableWidget();
+
+	bool IsWidgetEnabled() const { return TickHandle.IsValid(); }
+
+	const FName& GetWidgetName() const { return WidgetName; }
+
+protected:
+	virtual void DrawWidget(float DeltaTime) = 0;
+
+private:
+	FDelegateHandle TickHandle;
+	FName WidgetName;
+};
+
+/**
+ * Specialized version of FSlateIMWidgetBase that auto registers a command to toggle this widget
+ */
+class SLATEIM_API FSlateIMWidgetWithCommandBase : public FSlateIMWidgetBase
+{
+public:
+	FSlateIMWidgetWithCommandBase(const TCHAR* Command, const TCHAR* CommandHelp);
+
+private:
+	FAutoConsoleCommand WidgetCommand;
+};
+
+/**
+ * Specialized version of FSlateIMWidgetBase that draws in a window and disables itself when the user closes the window
+ */
+class SLATEIM_API FSlateIMWindowBase : public FSlateIMWidgetWithCommandBase
+{
+public:
+	FSlateIMWindowBase(FStringView WindowTitle, FVector2f WindowSize, const TCHAR* Command, const TCHAR* CommandHelp)
+		: FSlateIMWidgetWithCommandBase(Command, CommandHelp)
+		, WindowTitle(WindowTitle)
+		, WindowSize(WindowSize)
+	{}
+
+protected:
+	virtual void DrawWindow(float DeltaTime) = 0;
+
+private:
+	virtual void DrawWidget(float DeltaTime) final override;
+
+	FString WindowTitle;
+	FVector2f WindowSize;
+};
+
+/**
+ * Specialized version of FSlateIMWidgetBase that exposes the resulting widget for embedding in other slate widgets
+ *
+ * @note EnableWidget() must be called to get a valid ExposedWidget
+ */
+class SLATEIM_API FSlateIMExposedBase : public FSlateIMWidgetBase
+{
+public:
+	explicit FSlateIMExposedBase(const FStringView& Name)
+		: FSlateIMWidgetBase(Name)
+	{
+	}
+
+	// Get the widget to embed in your existing slate hierarchy
+	TSharedRef<SWidget> GetExposedWidget() const;
+
+	DECLARE_MULTICAST_DELEGATE_OneParam(FOnExposedWidgetChanged, TSharedRef<SWidget>);
+	FOnExposedWidgetChanged OnExposedWidgetChanged;
+
+protected:
+	virtual void DrawContent(float DeltaTime) = 0;
+
+private:
+	virtual void DrawWidget(float DeltaTime) final override;
+
+	TSharedPtr<SWidget> ExposedWidget;
+	
+};
diff --git a/Engine/Plugins/Experimental/SlateIM/Source/SlateIM/SlateIM.Build.cs b/Engine/Plugins/Experimental/SlateIM/Source/SlateIM/SlateIM.Build.cs
new file mode 100644
--- /dev/null
+++ b/Engine/Plugins/Experimental/SlateIM/Source/SlateIM/SlateIM.Build.cs
@@ -0,0 +1,55 @@
+// Copyright Epic Games, Inc. All Rights Reserved.
+
+using UnrealBuildTool;
+
+public class SlateIM : ModuleRules
+{
+	public SlateIM(ReadOnlyTargetRules Target) : base(Target)
+	{
+		PCHUsage = ModuleRules.PCHUsageMode.UseExplicitOrSharedPCHs;
+		
+		PublicDependencyModuleNames.AddRange(
+			new string[]
+			{
+				"Core",
+			}
+		);
+			
+		PrivateDependencyModuleNames.AddRange(
+			new string[]
+			{
+				"CoreUObject",
+				"Slate",
+				"SlateCore",
+				"InputCore",
+			}
+		);
+
+		if (Target.bCompileAgainstEngine)
+		{
+			PrivateDependencyModuleNames.AddRange(
+				new string[]
+				{
+					"Engine",
+				});
+		}
+
+		if(Target.bCompileAgainstEditor)
+		{
+			PrivateDependencyModuleNames.AddRange(
+				new string[]
+				{
+					"UnrealEd",
+				});
+		}
+		
+		if (Target.Configuration != UnrealTargetConfiguration.Shipping)
+		{
+			PublicDefinitions.Add("WITH_SLATEIM_EXAMPLES=1");
+		}
+		else
+		{
+			PublicDefinitions.Add("WITH_SLATEIM_EXAMPLES=0");
+		}
+	}
+}
