SlateIM: Add a style browser window to search the list of pre-defined styles available for use in SlateIM

Open with the command SlateIM.ToggleSlateStyleBrowser or from the test window

#rnx

diff --git a/Engine/Plugins/Experimental/SlateIM/Source/SlateIM/Private/Misc/SlateIMExamples.cpp b/Engine/Plugins/Experimental/SlateIM/Source/SlateIM/Private/Misc/SlateIMExamples.cpp
--- a/Engine/Plugins/Experimental/SlateIM/Source/SlateIM/Private/Misc/SlateIMExamples.cpp
+++ b/Engine/Plugins/Experimental/SlateIM/Source/SlateIM/Private/Misc/SlateIMExamples.cpp
@@ -3,11 +3,13 @@
 #if WITH_SLATEIM_EXAMPLES
 #include "SlateIMExamples.h"
 
+#include "HAL/PlatformApplicationMisc.h"
 #include "Misc/App.h"
 #include "ProfilingDebugging/ScopedTimers.h"
 #include "SlateIM.h"
 #include "SlateIMLogging.h"
 #include "Styling/AppStyle.h"
+#include "Styling/SlateStyle.h"
 #include "Styling/StyleColors.h"
 
 #if WITH_EDITOR
@@ -18,6 +20,53 @@
 #include "Engine/Texture2D.h"
 #endif
 
+namespace SlateIM::Private
+{
+	struct FExposedSlateStyle : public FSlateStyleSet
+	{
+	private:
+		FExposedSlateStyle()
+			: FSlateStyleSet(TEXT("ExposedSlateStyleSet"))
+		{}
+		
+	public:
+		TArray<FName> GetBrushStyleKeys(const FString& SearchString) const
+		{
+			TArray<FName> Keys;
+			BrushResources.GenerateKeyArray(Keys);
+
+			if (SearchString.IsEmpty())
+			{
+				return Keys;
+			}
+			
+			return Keys.FilterByPredicate([&SearchString](const FName& Key)
+			{
+				return Key.ToString().ToLower().Contains(SearchString);
+			});
+		}
+
+		template<typename WidgetStyle>
+		TMap<FName, const WidgetStyle*> GetWidgetStyles(const FString& SearchString) const
+		{
+			TMap<FName, const WidgetStyle*> Styles;
+
+			for (const auto& StylePair : WidgetStyleValues)
+			{
+				if (StylePair.Value->GetTypeName() == WidgetStyle::TypeName)
+				{
+					if (SearchString.IsEmpty() || StylePair.Key.ToString().ToLower().Contains(SearchString))
+					{
+						Styles.Add(StylePair.Key, static_cast<const WidgetStyle*>(&StylePair.Value.Get()));
+					}
+				}
+			}
+
+			return Styles;
+		}
+	};
+}
+
 void FSlateIMTestWidget::Draw()
 {
 	double LastTime = CurrentTime;
@@ -34,9 +83,22 @@ void FSlateIMTestWidget::Draw()
 		TimeText = FString::Printf(TEXT("%.3f ms"), LastTime * 1000);
 		TimeSinceLastUpdate = 0;
 	}
-	SlateIM::HAlign(HAlign_Right);
-	SlateIM::Padding({ 5.f, 5.f, 5.f, 0.f });
-	SlateIM::Text(TimeText);
+	SlateIM::HAlign(HAlign_Fill);
+	SlateIM::BeginHorizontalStack();
+	{
+#if WITH_ENGINE
+		if (GEngine && SlateIM::Button(TEXT("Open Style Browser")))
+		{
+			GEngine->Exec(nullptr, TEXT("SlateIM.ToggleSlateStyleBrowser"));
+		}
+#endif
+		
+		SlateIM::Fill();
+		SlateIM::HAlign(HAlign_Right);
+		SlateIM::Padding({ 5.f, 5.f, 5.f, 0.f });
+		SlateIM::Text(TimeText);
+	}
+	SlateIM::EndHorizontalStack();
 	
 	// Full root scrollbox
 	SlateIM::Fill();
@@ -586,6 +648,343 @@ void FSlateIMTestWidget::Draw()
 	SlateIM::EndBorder();
 }
 
+void FSlateStyleBrowser::DrawWindow(float DeltaTime)
+{
+	static const TArray<FString> Options = { TEXT("Option 1"), TEXT("Option 2"), TEXT("Option 3") };
+	SlateIM::HAlign(HAlign_Fill);
+	SlateIM::VAlign(VAlign_Fill);
+	SlateIM::Fill();
+	SlateIM::BeginVerticalStack();
+	SlateIM::EditableText(SearchString, TEXT("Search Styles..."));
+	
+	SlateIM::HAlign(HAlign_Fill);
+	SlateIM::VAlign(VAlign_Fill);
+	SlateIM::Fill();
+	SlateIM::BeginTable();
+	SlateIM::AddTableColumn(TEXT("Name"));
+	SlateIM::AddTableColumn(TEXT("Preview"));
+	SlateIM::InitialTableColumnWidth(80.f);
+	SlateIM::AddTableColumn();
+
+	const auto& Style = static_cast<const SlateIM::Private::FExposedSlateStyle&>(FAppStyle::Get());
+	const FString LowerCaseSearchString = SearchString.ToLower();
+
+	auto DrawNameCell = [](const FName& Key)
+	{
+		if (SlateIM::NextTableCell())
+		{
+			SlateIM::VAlign(VAlign_Center);
+			SlateIM::Fill();
+			SlateIM::Text(Key.ToString());
+		}
+	};
+
+	auto DrawCopyCell = [](const FName& Key)
+	{
+		if (SlateIM::NextTableCell())
+		{
+			SlateIM::SetToolTip(TEXT("Click to copy the style name to your clipboard"));
+			SlateIM::HAlign(HAlign_Center);
+			SlateIM::VAlign(VAlign_Center);
+			SlateIM::Fill();
+			if (SlateIM::Button(TEXT("Copy")))
+			{
+				FPlatformApplicationMisc::ClipboardCopy(*Key.ToString());
+			}
+		}
+	};
+
+	// Brushes
+	{
+		if (SlateIM::NextTableCell())
+		{
+			SlateIM::Text(TEXT("Brushes"));
+		}
+		SlateIM::NextTableCell(); // Skip Preview Column
+		SlateIM::NextTableCell(); // Skip Button Column
+
+		if (SlateIM::BeginTableRowChildren())
+		{
+			for (const FName& BrushStyleKey : Style.GetBrushStyleKeys(LowerCaseSearchString))
+			{
+				if (const FSlateBrush* Brush = Style.GetBrush(BrushStyleKey))
+				{
+					DrawNameCell(BrushStyleKey);
+
+					if (SlateIM::NextTableCell())
+					{
+						SlateIM::HAlign(HAlign_Center);
+						SlateIM::VAlign(VAlign_Center);
+						SlateIM::Fill();
+						SlateIM::Image(Brush);
+					}
+					
+					DrawCopyCell(BrushStyleKey);
+				}
+			}
+		}
+		SlateIM::EndTableRowChildren();
+	}
+
+	// Text Block Styles
+	{
+		if (SlateIM::NextTableCell())
+		{
+			SlateIM::Text(TEXT("Text Block Styles"));
+		}
+		SlateIM::NextTableCell(); // Skip Preview Column
+		SlateIM::NextTableCell(); // Skip Button Column
+		if (SlateIM::BeginTableRowChildren())
+		{
+			for (const TPair<FName, const FTextBlockStyle*>& WidgetStyle : Style.GetWidgetStyles<FTextBlockStyle>(LowerCaseSearchString))
+			{
+				DrawNameCell(WidgetStyle.Key);
+
+				if (SlateIM::NextTableCell())
+				{
+					SlateIM::HAlign(HAlign_Center);
+					SlateIM::VAlign(VAlign_Center);
+					SlateIM::Fill();
+					SlateIM::Text(TEXT("The quick brown fox jumps over the lazy dog."), WidgetStyle.Value);
+				}
+			
+				DrawCopyCell(WidgetStyle.Key);
+			}
+		}
+		SlateIM::EndTableRowChildren();
+	}
+
+	// Editable Text Box Styles
+	{
+		if (SlateIM::NextTableCell())
+		{
+			SlateIM::Text(TEXT("Editable Text Box Styles"));
+		}
+		SlateIM::NextTableCell(); // Skip Preview Column
+		SlateIM::NextTableCell(); // Skip Button Column
+		if (SlateIM::BeginTableRowChildren())
+		{
+			for (const TPair<FName, const FEditableTextBoxStyle*>& WidgetStyle : Style.GetWidgetStyles<FEditableTextBoxStyle>(LowerCaseSearchString))
+			{
+				DrawNameCell(WidgetStyle.Key);
+
+				if (SlateIM::NextTableCell())
+				{
+					SlateIM::HAlign(HAlign_Center);
+					SlateIM::VAlign(VAlign_Center);
+					SlateIM::Fill();
+					SlateIM::EditableText(PreviewText, TEXT("Hint text..."), WidgetStyle.Value);
+				}
+				
+				DrawCopyCell(WidgetStyle.Key);
+			}
+		}
+		SlateIM::EndTableRowChildren();
+	}
+
+	// Button Styles
+	{
+		if (SlateIM::NextTableCell())
+		{
+			SlateIM::Text(TEXT("Button Styles"));
+		}
+		SlateIM::NextTableCell(); // Skip Preview Column
+		SlateIM::NextTableCell(); // Skip Button Column
+		if (SlateIM::BeginTableRowChildren())
+		{
+			for (const TPair<FName, const FButtonStyle*>& WidgetStyle : Style.GetWidgetStyles<FButtonStyle>(LowerCaseSearchString))
+			{
+				DrawNameCell(WidgetStyle.Key);
+
+				if (SlateIM::NextTableCell())
+				{
+					SlateIM::HAlign(HAlign_Center);
+					SlateIM::VAlign(VAlign_Center);
+					SlateIM::Fill();
+					SlateIM::Button(TEXT("Click Me"), WidgetStyle.Value);
+				}
+				
+				DrawCopyCell(WidgetStyle.Key);
+			}
+		}
+		SlateIM::EndTableRowChildren();
+	}
+
+	// SpinBox Styles
+	{
+		if (SlateIM::NextTableCell())
+		{
+			SlateIM::Text(TEXT("SpinBox Styles"));
+		}
+		SlateIM::NextTableCell(); // Skip Preview Column
+		SlateIM::NextTableCell(); // Skip Button Column
+		if (SlateIM::BeginTableRowChildren())
+		{
+			for (const TPair<FName, const FSpinBoxStyle*>& WidgetStyle : Style.GetWidgetStyles<FSpinBoxStyle>(LowerCaseSearchString))
+			{
+				DrawNameCell(WidgetStyle.Key);
+
+				if (SlateIM::NextTableCell())
+				{
+					SlateIM::HAlign(HAlign_Center);
+					SlateIM::VAlign(VAlign_Center);
+					SlateIM::Fill();
+					SlateIM::SpinBox(SpinBoxValue, -100.f, 100.f, WidgetStyle.Value);
+				}
+				
+				DrawCopyCell(WidgetStyle.Key);
+			}
+		}
+		SlateIM::EndTableRowChildren();
+	}
+
+	// Slider Styles
+	{
+		if (SlateIM::NextTableCell())
+		{
+			SlateIM::Text(TEXT("Slider Styles"));
+		}
+		SlateIM::NextTableCell(); // Skip Preview Column
+		SlateIM::NextTableCell(); // Skip Button Column
+		if (SlateIM::BeginTableRowChildren())
+		{
+			for (const TPair<FName, const FSliderStyle*>& WidgetStyle : Style.GetWidgetStyles<FSliderStyle>(LowerCaseSearchString))
+			{
+				DrawNameCell(WidgetStyle.Key);
+
+				if (SlateIM::NextTableCell())
+				{
+					SlateIM::HAlign(HAlign_Center);
+					SlateIM::VAlign(VAlign_Center);
+					SlateIM::Fill();
+					SlateIM::Slider(SliderValue, 0.f, 100.f, 0.1f, WidgetStyle.Value);
+				}
+				
+				DrawCopyCell(WidgetStyle.Key);
+			}
+		}
+		SlateIM::EndTableRowChildren();
+	}
+
+	// ProgressBar Styles
+	{
+		if (SlateIM::NextTableCell())
+		{
+			SlateIM::Text(TEXT("Progress Bar Styles"));
+		}
+		SlateIM::NextTableCell(); // Skip Preview Column
+		SlateIM::NextTableCell(); // Skip Button Column
+		if (SlateIM::BeginTableRowChildren())
+		{
+			for (const TPair<FName, const FProgressBarStyle*>& WidgetStyle : Style.GetWidgetStyles<FProgressBarStyle>(LowerCaseSearchString))
+			{
+				DrawNameCell(WidgetStyle.Key);
+
+				if (SlateIM::NextTableCell())
+				{
+					SlateIM::HAlign(HAlign_Center);
+					SlateIM::VAlign(VAlign_Center);
+					SlateIM::Fill();
+					SlateIM::ProgressBar(SliderValue / 100.f, WidgetStyle.Value);
+				}
+				
+				DrawCopyCell(WidgetStyle.Key);
+			}
+		}
+		SlateIM::EndTableRowChildren();
+	}
+
+	// ComboBox Styles
+	{
+		if (SlateIM::NextTableCell())
+		{
+			SlateIM::Text(TEXT("Combo Box Styles"));
+		}
+		SlateIM::NextTableCell(); // Skip Preview Column
+		SlateIM::NextTableCell(); // Skip Button Column
+		if (SlateIM::BeginTableRowChildren())
+		{
+			for (const TPair<FName, const FComboBoxStyle*>& WidgetStyle : Style.GetWidgetStyles<FComboBoxStyle>(LowerCaseSearchString))
+			{
+				DrawNameCell(WidgetStyle.Key);
+
+				if (SlateIM::NextTableCell())
+				{
+					SlateIM::HAlign(HAlign_Center);
+					SlateIM::VAlign(VAlign_Center);
+					SlateIM::Fill();
+					SlateIM::ComboBox(Options, SelectedComboIndex, WidgetStyle.Value);
+				}
+				
+				DrawCopyCell(WidgetStyle.Key);
+			}
+		}
+		SlateIM::EndTableRowChildren();
+	}
+
+	// Table View Styles
+	{
+		if (SlateIM::NextTableCell())
+		{
+			SlateIM::Text(TEXT("Table View Styles"));
+		}
+		SlateIM::NextTableCell(); // Skip Preview Column
+		SlateIM::NextTableCell(); // Skip Button Column
+		if (SlateIM::BeginTableRowChildren())
+		{
+			for (const TPair<FName, const FTableViewStyle*>& WidgetStyle : Style.GetWidgetStyles<FTableViewStyle>(LowerCaseSearchString))
+			{
+				DrawNameCell(WidgetStyle.Key);
+
+				if (SlateIM::NextTableCell())
+				{
+					SlateIM::HAlign(HAlign_Center);
+					SlateIM::VAlign(VAlign_Center);
+					SlateIM::Fill();
+					SlateIM::BeginHorizontalStack();
+					{
+						SlateIM::HAlign(HAlign_Center);
+						SlateIM::VAlign(VAlign_Top);
+						SlateIM::Fill();
+						SlateIM::SelectionList(Options, SelectedListIndex, false, WidgetStyle.Value);
+
+						SlateIM::HAlign(HAlign_Center);
+						SlateIM::VAlign(VAlign_Top);
+						SlateIM::Fill();
+						SlateIM::BeginTable(WidgetStyle.Value);
+						SlateIM::AddTableColumn(TEXT("Column 1"));
+						SlateIM::AddTableColumn(TEXT("Column 2"));
+						SlateIM::NextTableCell();
+						SlateIM::Text(TEXT("Cell 1"));
+						SlateIM::NextTableCell();
+						SlateIM::Text(TEXT("Cell 2"));
+						SlateIM::NextTableCell();
+						SlateIM::Text(TEXT("Cell 3"));
+						SlateIM::NextTableCell();
+						SlateIM::Text(TEXT("Cell 4"));
+						if (SlateIM::BeginTableRowChildren())
+						{
+							SlateIM::NextTableCell();
+							SlateIM::Text(TEXT("Cell 5"));
+							SlateIM::NextTableCell();
+							SlateIM::Text(TEXT("Cell 6"));
+						}
+						SlateIM::EndTableRowChildren();
+						SlateIM::EndTable();
+					}
+					SlateIM::EndHorizontalStack();
+				}
+				
+				DrawCopyCell(WidgetStyle.Key);
+			}
+		}
+		SlateIM::EndTableRowChildren();
+	}
+	
+	SlateIM::EndTable();
+	SlateIM::EndVerticalStack();
+}
+
 void FSlateIMTestWindowWidget::DrawWindow(float DeltaTime)
 {
 	TestWidget.Draw();
diff --git a/Engine/Plugins/Experimental/SlateIM/Source/SlateIM/Private/Misc/SlateIMExamples.h b/Engine/Plugins/Experimental/SlateIM/Source/SlateIM/Private/Misc/SlateIMExamples.h
--- a/Engine/Plugins/Experimental/SlateIM/Source/SlateIM/Private/Misc/SlateIMExamples.h
+++ b/Engine/Plugins/Experimental/SlateIM/Source/SlateIM/Private/Misc/SlateIMExamples.h
@@ -42,6 +42,24 @@ private:
 #endif
 };
 
+class FSlateStyleBrowser : public FSlateIMWindowBase
+{
+public:
+	FSlateStyleBrowser()
+		: FSlateIMWindowBase(TEXT("SlateIM Style Browser"), FVector2f(1000, 500), TEXT("SlateIM.ToggleSlateStyleBrowser"), TEXT("Opens a window that previews available slate styles"))
+	{}
+	
+	virtual void DrawWindow(float DeltaTime) override;
+
+private:
+	FString SearchString;
+	FString PreviewText;
+	float SpinBoxValue = 66.7f;
+	float SliderValue = 66.7f;
+	int32 SelectedComboIndex = INDEX_NONE;
+	int32 SelectedListIndex = INDEX_NONE;
+};
+
 class FSlateIMTestWindowWidget : public FSlateIMWindowBase
 {
 public:
diff --git a/Engine/Plugins/Experimental/SlateIM/Source/SlateIM/Private/Misc/SlateIMManager.h b/Engine/Plugins/Experimental/SlateIM/Source/SlateIM/Private/Misc/SlateIMManager.h
--- a/Engine/Plugins/Experimental/SlateIM/Source/SlateIM/Private/Misc/SlateIMManager.h
+++ b/Engine/Plugins/Experimental/SlateIM/Source/SlateIM/Private/Misc/SlateIMManager.h
@@ -211,6 +211,7 @@ namespace SlateIM
 		bool bIsSlateIMModalOpen = false;
 
 #if WITH_SLATEIM_EXAMPLES
+		FSlateStyleBrowser SlateStyleBrowser;
 		FSlateIMTestWindowWidget TestWindowWidget;
 #if WITH_ENGINE
 		FSlateIMTestViewportWidget TestViewportWidget;
diff --git a/Engine/Plugins/Experimental/SlateIM/Source/SlateIM/SlateIM.Build.cs b/Engine/Plugins/Experimental/SlateIM/Source/SlateIM/SlateIM.Build.cs
--- a/Engine/Plugins/Experimental/SlateIM/Source/SlateIM/SlateIM.Build.cs
+++ b/Engine/Plugins/Experimental/SlateIM/Source/SlateIM/SlateIM.Build.cs
@@ -18,7 +18,7 @@ public class SlateIM : ModuleRules
 		PrivateDependencyModuleNames.AddRange(
 			new string[]
 			{
-				"CoreUObject",
+				"ApplicationCore",
 				"Slate",
 				"SlateCore",
 				"InputCore",
@@ -30,6 +30,7 @@ public class SlateIM : ModuleRules
 			PrivateDependencyModuleNames.AddRange(
 				new string[]
 				{
+					"CoreUObject",
 					"Engine",
 				});
 		}
