SlateIM: Add PopUp widget to display floating UI anchored to other UI

#rnx

diff --git a/Engine/Plugins/Experimental/SlateIM/Source/SlateIM/Private/Containers/SImPopUp.cpp b/Engine/Plugins/Experimental/SlateIM/Source/SlateIM/Private/Containers/SImPopUp.cpp
new file mode 100644
--- /dev/null
+++ b/Engine/Plugins/Experimental/SlateIM/Source/SlateIM/Private/Containers/SImPopUp.cpp
@@ -0,0 +1,76 @@
+﻿// Copyright Epic Games, Inc. All Rights Reserved.
+
+#include "SImPopUp.h"
+
+#include "Misc/SlateIMSlotData.h"
+#include "Widgets/Layout/SBox.h"
+
+SLATE_IMPLEMENT_WIDGET(SImPopUp)
+
+void SImPopUp::PrivateRegisterAttributes(FSlateAttributeInitializer& AttributeInitializer)
+{
+}
+
+void SImPopUp::Construct(const FArguments& InArgs)
+{
+	ChildBox = SNew(SBox)
+		[
+			InArgs._Content.Widget
+		];
+	
+	SMenuAnchor::Construct(FArguments(InArgs)
+		.UseApplicationMenuStack(false)
+		.MenuContent(ChildBox.ToSharedRef()));
+}
+
+int32 SImPopUp::GetNumChildren()
+{
+	if (ChildBox.IsValid() && ChildBox->GetChildren() && ChildBox->GetChildren()->Num() > 0)
+	{
+		TSharedRef<SWidget> ChildWidget = ChildBox->GetChildren()->GetChildAt(0);
+
+		return (ChildWidget == SNullWidget::NullWidget) ? 0 : 1;
+	}
+
+	return 0;
+}
+
+FSlateIMChild SImPopUp::GetChild(int32 Index)
+{
+	if (ChildBox.IsValid() && ChildBox->GetChildren())
+	{
+		return ChildBox->GetChildren()->GetChildAt(Index);
+	}
+
+	return nullptr;
+}
+
+void SImPopUp::UpdateChild(FSlateIMChild Child, int32 Index, const FSlateIMSlotData& AlignmentData)
+{
+	TSharedPtr<SWidget> ChildWidget = Child.GetWidget();
+	if (ensureMsgf(ChildWidget, TEXT("Invalid child in SlateIM Popup")))
+	{
+		ChildBox->SetPadding(AlignmentData.Padding);
+		ChildBox->SetHAlign(AlignmentData.HorizontalAlignment);
+		ChildBox->SetVAlign(AlignmentData.VerticalAlignment);
+		ChildBox->SetMinDesiredWidth(AlignmentData.MinWidth > 0 ? AlignmentData.MinWidth : FOptionalSize());
+		ChildBox->SetMinDesiredHeight(AlignmentData.MinHeight > 0 ? AlignmentData.MinHeight : FOptionalSize());
+		ChildBox->SetMaxDesiredWidth(AlignmentData.MaxWidth > 0 ? AlignmentData.MaxWidth : FOptionalSize());
+		ChildBox->SetMaxDesiredHeight(AlignmentData.MaxHeight > 0 ? AlignmentData.MaxHeight : FOptionalSize());
+		ChildBox->SetContent(ChildWidget.ToSharedRef());
+	}
+}
+
+void SImPopUp::RemoveUnusedChildren(int32 LastUsedChildIndex)
+{
+	if (ChildBox.IsValid())
+	{
+		ChildBox->SetContent(SNullWidget::NullWidget);
+	}
+}
+
+FVector2D SImPopUp::ComputeDesiredSize(float LayoutScaleMultiplier) const
+{
+	return FVector2D::ZeroVector;
+}
+
diff --git a/Engine/Plugins/Experimental/SlateIM/Source/SlateIM/Private/Containers/SImPopUp.h b/Engine/Plugins/Experimental/SlateIM/Source/SlateIM/Private/Containers/SImPopUp.h
new file mode 100644
--- /dev/null
+++ b/Engine/Plugins/Experimental/SlateIM/Source/SlateIM/Private/Containers/SImPopUp.h
@@ -0,0 +1,33 @@
+﻿// Copyright Epic Games, Inc. All Rights Reserved.
+
+#pragma once
+
+#include "ISlateIMContainer.h"
+#include "Widgets/Input/SMenuAnchor.h"
+
+class SBox;
+
+class SLATEIM_API SImPopUp : public SMenuAnchor, public ISlateIMContainer
+{
+	SLATE_DECLARE_WIDGET(SImPopUp, SMenuAnchor)
+	SLATE_IM_TYPE_DATA(SImPopUp, ISlateIMContainer)
+	
+public:
+	void Construct(const FArguments& InArgs);
+	
+	virtual int32 GetNumChildren() override;
+	virtual FSlateIMChild GetChild(int32 Index) override;
+	virtual void UpdateChild(FSlateIMChild Child, int32 Index, const FSlateIMSlotData& AlignmentData) override;
+
+	virtual FSlateIMChild GetContainer() override
+	{
+		return AsShared();
+	}
+
+	virtual void RemoveUnusedChildren(int32 LastUsedChildIndex) override;
+
+	virtual FVector2D ComputeDesiredSize(float LayoutScaleMultiplier) const override;
+
+private:
+	TSharedPtr<SBox> ChildBox;
+};
diff --git a/Engine/Plugins/Experimental/SlateIM/Source/SlateIM/Private/Implementation/SlateIM_Containers.cpp b/Engine/Plugins/Experimental/SlateIM/Source/SlateIM/Private/Implementation/SlateIM_Containers.cpp
--- a/Engine/Plugins/Experimental/SlateIM/Source/SlateIM/Private/Implementation/SlateIM_Containers.cpp
+++ b/Engine/Plugins/Experimental/SlateIM/Source/SlateIM/Private/Implementation/SlateIM_Containers.cpp
@@ -4,6 +4,7 @@
 
 #include "Containers/SImCompoundWidget.h"
 #include "Containers/SImContextMenuAnchor.h"
+#include "Containers/SImPopUp.h"
 #include "Containers/SImScrollBox.h"
 #include "Containers/SImStackBox.h"
 #include "Containers/SImTableView.h"
@@ -181,7 +182,41 @@ namespace SlateIM
 	{
 		FSlateIMManager::Get().PopContainer<SImScrollBox>();
 	}
-	
+
+	void BeginPopUp(const FName BorderStyleName, EOrientation Orientation, bool bAbsorbMouse, FMargin ContentPadding)
+	{
+		BeginPopUp(FAppStyle::GetBrush(BorderStyleName), Orientation, bAbsorbMouse, ContentPadding);
+	}
+
+	void BeginPopUp(const FSlateBrush* BorderBrush, EOrientation Orientation, bool bAbsorbMouse, FMargin ContentPadding)
+	{
+		TSharedPtr<SImPopUp> ContainerWidget;
+		{
+			FWidgetScope<SImPopUp> Scope(FMargin(0));
+			ContainerWidget = Scope.GetWidget();
+
+			if (!ContainerWidget)
+			{
+				ContainerWidget =
+					SNew(SImPopUp)
+					.ShowMenuBackground(false); 
+				
+				Scope.UpdateWidget(ContainerWidget);
+			}
+
+			ContainerWidget->SetIsOpen(true);
+		}
+
+		FSlateIMManager::Get().PushContainer(FContainerNode(ContainerWidget));
+		SlateIM::BeginBorder(BorderBrush, Orientation, bAbsorbMouse, ContentPadding);
+	}
+
+	void EndPopUp()
+	{
+		SlateIM::EndBorder();
+		FSlateIMManager::Get().PopContainer<SImPopUp>();
+	}
+
 	void BeginTableRow()
 	{
 		// If we're in a row, we're adding a child row, otherwise it's a top-level row in the table
diff --git a/Engine/Plugins/Experimental/SlateIM/Source/SlateIM/Private/Misc/SlateIMExamples.cpp b/Engine/Plugins/Experimental/SlateIM/Source/SlateIM/Private/Misc/SlateIMExamples.cpp
--- a/Engine/Plugins/Experimental/SlateIM/Source/SlateIM/Private/Misc/SlateIMExamples.cpp
+++ b/Engine/Plugins/Experimental/SlateIM/Source/SlateIM/Private/Misc/SlateIMExamples.cpp
@@ -247,6 +247,22 @@ void FSlateIMTestWidget::Draw()
 			SlateIM::EndHorizontalStack();
 		}
 		
+		{
+			SCOPED_NAMED_EVENT_TEXT("PopUp example", FColorList::Goldenrod);
+			// PopUp example
+			SlateIM::Padding(0);
+			SlateIM::BeginHorizontalStack();
+			SlateIM::CheckBox(TEXT("Show a floating popup"), bShowPopUp);
+			if (bShowPopUp)
+			{
+				SlateIM::BeginPopUp();
+				SlateIM::Text(TEXT("Pop Up Test:"));
+				SlateIM::Image("AppIcon");
+				SlateIM::EndPopUp();
+			}
+			SlateIM::EndHorizontalStack();
+		}
+		
 		{
 			SCOPED_NAMED_EVENT_TEXT("DisabledState example", FColorList::Goldenrod);
 			// DisabledState example
diff --git a/Engine/Plugins/Experimental/SlateIM/Source/SlateIM/Private/Misc/SlateIMExamples.h b/Engine/Plugins/Experimental/SlateIM/Source/SlateIM/Private/Misc/SlateIMExamples.h
--- a/Engine/Plugins/Experimental/SlateIM/Source/SlateIM/Private/Misc/SlateIMExamples.h
+++ b/Engine/Plugins/Experimental/SlateIM/Source/SlateIM/Private/Misc/SlateIMExamples.h
@@ -26,6 +26,7 @@ private:
 	int32 SelectedItem = INDEX_NONE;
 	int32 IntValue = 50;
 	int32 IntMax = 100;
+	bool bShowPopUp = false;
 	bool bShouldBeDisabled = false;
 	bool MenuCheckState = true;
 	bool MenuToggleState = true;
diff --git a/Engine/Plugins/Experimental/SlateIM/Source/SlateIM/Public/SlateIM.h b/Engine/Plugins/Experimental/SlateIM/Source/SlateIM/Public/SlateIM.h
--- a/Engine/Plugins/Experimental/SlateIM/Source/SlateIM/Public/SlateIM.h
+++ b/Engine/Plugins/Experimental/SlateIM/Source/SlateIM/Public/SlateIM.h
@@ -233,6 +233,10 @@ namespace SlateIM
 	 */
 	SLATEIM_API bool BeginScrollBox(EOrientation ScrollBarOrientation = EOrientation::Orient_Vertical);
 	SLATEIM_API void EndScrollBox();
+
+	SLATEIM_API void BeginPopUp(const FName BorderStyleName = TEXT("ToolPanel.GroupBorder"), EOrientation Orientation = Orient_Vertical, bool bAbsorbMouse = true, FMargin ContentPadding = FMargin(2.0f));
+	SLATEIM_API void BeginPopUp(const FSlateBrush* BorderBrush, EOrientation Orientation = Orient_Vertical, bool bAbsorbMouse = true, FMargin ContentPadding = FMargin(2.0f));
+	SLATEIM_API void EndPopUp();
 #pragma endregion Containers
 
 #pragma region Slots
