SlateIM: Fix crash caused by setting a null style when invalidating a SpinBox

#rnx

diff --git a/Engine/Plugins/Experimental/SlateIM/Source/SlateIM/Private/Implementation/SlateIM_Widgets.cpp b/Engine/Plugins/Experimental/SlateIM/Source/SlateIM/Private/Implementation/SlateIM_Widgets.cpp
--- a/Engine/Plugins/Experimental/SlateIM/Source/SlateIM/Private/Implementation/SlateIM_Widgets.cpp
+++ b/Engine/Plugins/Experimental/SlateIM/Source/SlateIM/Private/Implementation/SlateIM_Widgets.cpp
@@ -440,7 +440,7 @@ namespace SlateIM
 				}
 			}
 
-			if (Scope.IsDataHashDirty())
+			if (SpinBoxStyle && Scope.IsDataHashDirty())
 			{
 				SpinBoxWidget->SetWidgetStyle(SpinBoxStyle);
 				SpinBoxWidget->InvalidateStyle();
