SlateIM: Add IsHovered() and IsFocused() functions to query the last rendered widget's state

- Also don't auto-focus PopUps so that they can open in response to other widgets receiving focus without auto-closing

#rnx

diff --git a/Engine/Plugins/Experimental/SlateIM/Source/SlateIM/Private/Containers/SImTableView.cpp b/Engine/Plugins/Experimental/SlateIM/Source/SlateIM/Private/Containers/SImTableView.cpp
--- a/Engine/Plugins/Experimental/SlateIM/Source/SlateIM/Private/Containers/SImTableView.cpp
+++ b/Engine/Plugins/Experimental/SlateIM/Source/SlateIM/Private/Containers/SImTableView.cpp
@@ -44,6 +44,16 @@ void FSlateIMTableRow::UpdateChild(FSlateIMChild Child, int32 Index, const FSlat
 	}
 }
 
+TSharedPtr<SWidget> FSlateIMTableRow::GetAsWidget()
+{
+	if (OwningTable)
+	{
+		return StaticCastSharedPtr<SImTableRow>(OwningTable->WidgetFromItem(AsShared()));
+	}
+
+	return nullptr;
+}
+
 void FSlateIMTableRow::GetChildRows(TArray<TSharedRef<FSlateIMTableRow>>& OutRows) const
 {
 	for (const FSlateIMChild& Child : Children)
diff --git a/Engine/Plugins/Experimental/SlateIM/Source/SlateIM/Private/Containers/SImTableView.h b/Engine/Plugins/Experimental/SlateIM/Source/SlateIM/Private/Containers/SImTableView.h
--- a/Engine/Plugins/Experimental/SlateIM/Source/SlateIM/Private/Containers/SImTableView.h
+++ b/Engine/Plugins/Experimental/SlateIM/Source/SlateIM/Private/Containers/SImTableView.h
@@ -22,6 +22,7 @@ public:
 	virtual FSlateIMChild GetContainer() override { return TSharedRef<ISlateIMChild>(AsShared()); }
 	virtual void RemoveUnusedChildren(int32 LastUsedChildIndex) override;
 	virtual void UpdateChild(FSlateIMChild Child, int32 Index, const FSlateIMSlotData& AlignmentData) override;
+	virtual TSharedPtr<SWidget> GetAsWidget() override;
 
 	void GetChildRows(TArray<TSharedRef<FSlateIMTableRow>>& OutRows) const;
 
diff --git a/Engine/Plugins/Experimental/SlateIM/Source/SlateIM/Private/Implementation/SlateIM_Containers.cpp b/Engine/Plugins/Experimental/SlateIM/Source/SlateIM/Private/Implementation/SlateIM_Containers.cpp
--- a/Engine/Plugins/Experimental/SlateIM/Source/SlateIM/Private/Implementation/SlateIM_Containers.cpp
+++ b/Engine/Plugins/Experimental/SlateIM/Source/SlateIM/Private/Implementation/SlateIM_Containers.cpp
@@ -204,7 +204,11 @@ namespace SlateIM
 				Scope.UpdateWidget(ContainerWidget);
 			}
 
-			ContainerWidget->SetIsOpen(true);
+			// Setting focus by default will cause the pop-up to auto-close when it loses focus which would mean we couldn't have multiple pop-ups at once
+			// nor could we open a pop-up in response to something else getting focus, so we don't focus the pop-up
+			constexpr bool bIsOpen = true;
+			constexpr bool bSetFocus = false;
+			ContainerWidget->SetIsOpen(bIsOpen, bSetFocus);
 		}
 
 		FSlateIMManager::Get().PushContainer(FContainerNode(ContainerWidget));
diff --git a/Engine/Plugins/Experimental/SlateIM/Source/SlateIM/Private/Implementation/SlateIM_Queries.cpp b/Engine/Plugins/Experimental/SlateIM/Source/SlateIM/Private/Implementation/SlateIM_Queries.cpp
new file mode 100644
--- /dev/null
+++ b/Engine/Plugins/Experimental/SlateIM/Source/SlateIM/Private/Implementation/SlateIM_Queries.cpp
@@ -0,0 +1,36 @@
+﻿// Copyright Epic Games, Inc. All Rights Reserved.
+
+#include "Misc/SlateIMManager.h"
+#include "SlateIM.h"
+
+namespace SlateIM
+{
+	bool IsHovered()
+	{
+		if (TSharedPtr<SWidget> ChildWidget = FSlateIMManager::Get().GetCurrentChildAsWidget())
+		{
+			return ChildWidget->IsHovered();
+		}
+
+		return false;
+	}
+
+	bool IsFocused(EFocusDepth Depth)
+	{
+		if (TSharedPtr<SWidget> ChildWidget = FSlateIMManager::Get().GetCurrentChildAsWidget())
+		{
+			switch (Depth)
+			{
+			case EFocusDepth::SelfOnly:
+				return ChildWidget->HasAnyUserFocus().IsSet();
+			case EFocusDepth::IncludingDescendants:
+				return ChildWidget->HasAnyUserFocusOrFocusedDescendants();
+			default:
+				checkNoEntry();
+				break;
+			}
+		}
+
+		return false;
+	}
+}
diff --git a/Engine/Plugins/Experimental/SlateIM/Source/SlateIM/Private/Misc/ISlateIMChild.h b/Engine/Plugins/Experimental/SlateIM/Source/SlateIM/Private/Misc/ISlateIMChild.h
--- a/Engine/Plugins/Experimental/SlateIM/Source/SlateIM/Private/Misc/ISlateIMChild.h
+++ b/Engine/Plugins/Experimental/SlateIM/Source/SlateIM/Private/Misc/ISlateIMChild.h
@@ -13,6 +13,8 @@ class SWidget;
 // Interface for non-widget children in the SlateIM hierarchy (eg. for virtualized listview entries)
 class ISlateIMChild : public ISlateIMTypeChecking
 {
+public:
+	virtual TSharedPtr<SWidget> GetAsWidget() = 0;
 };
 
 struct FSlateIMChild : private TUnion<TYPE_OF_NULLPTR, TSharedRef<SWidget>, TSharedRef<ISlateIMChild>>
diff --git a/Engine/Plugins/Experimental/SlateIM/Source/SlateIM/Private/Misc/SlateIMExamples.cpp b/Engine/Plugins/Experimental/SlateIM/Source/SlateIM/Private/Misc/SlateIMExamples.cpp
--- a/Engine/Plugins/Experimental/SlateIM/Source/SlateIM/Private/Misc/SlateIMExamples.cpp
+++ b/Engine/Plugins/Experimental/SlateIM/Source/SlateIM/Private/Misc/SlateIMExamples.cpp
@@ -122,10 +122,21 @@ void FSlateIMTestWidget::Draw()
 				UE_LOG(LogSlateIM, Log, TEXT("Button was clicked"));
 			}
 
-			// EditableText Example
-			if (SlateIM::EditableText(TypedText, TEXT("Type In Here")))
 			{
-				UE_LOG(LogSlateIM, Log, TEXT("Editable Text Changed: %s"), *TypedText);
+				SlateIM::BeginVerticalStack();
+				// EditableText Example
+				if (SlateIM::EditableText(TypedText, TEXT("Type In Here")))
+				{
+					UE_LOG(LogSlateIM, Log, TEXT("Editable Text Changed: %s"), *TypedText);
+				}
+				if (SlateIM::IsFocused(SlateIM::EFocusDepth::IncludingDescendants))
+				{
+					SlateIM::BeginPopUp();
+					SlateIM::Text(TEXT("Pop Up Test:"));
+					SlateIM::Image("AppIcon");
+					SlateIM::EndPopUp();
+				}
+				SlateIM::EndVerticalStack();
 			}
 
 			// CheckBox examples
@@ -252,8 +263,8 @@ void FSlateIMTestWidget::Draw()
 			// PopUp example
 			SlateIM::Padding(0);
 			SlateIM::BeginHorizontalStack();
-			SlateIM::CheckBox(TEXT("Show a floating popup"), bShowPopUp);
-			if (bShowPopUp)
+			SlateIM::Text(TEXT("Hover here to Show a floating popup"));
+			if (SlateIM::IsHovered())
 			{
 				SlateIM::BeginPopUp();
 				SlateIM::Text(TEXT("Pop Up Test:"));
diff --git a/Engine/Plugins/Experimental/SlateIM/Source/SlateIM/Private/Misc/SlateIMExamples.h b/Engine/Plugins/Experimental/SlateIM/Source/SlateIM/Private/Misc/SlateIMExamples.h
--- a/Engine/Plugins/Experimental/SlateIM/Source/SlateIM/Private/Misc/SlateIMExamples.h
+++ b/Engine/Plugins/Experimental/SlateIM/Source/SlateIM/Private/Misc/SlateIMExamples.h
@@ -26,7 +26,6 @@ private:
 	int32 SelectedItem = INDEX_NONE;
 	int32 IntValue = 50;
 	int32 IntMax = 100;
-	bool bShowPopUp = false;
 	bool bShouldBeDisabled = false;
 	bool MenuCheckState = true;
 	bool MenuToggleState = true;
diff --git a/Engine/Plugins/Experimental/SlateIM/Source/SlateIM/Private/Misc/SlateIMManager.cpp b/Engine/Plugins/Experimental/SlateIM/Source/SlateIM/Private/Misc/SlateIMManager.cpp
--- a/Engine/Plugins/Experimental/SlateIM/Source/SlateIM/Private/Misc/SlateIMManager.cpp
+++ b/Engine/Plugins/Experimental/SlateIM/Source/SlateIM/Private/Misc/SlateIMManager.cpp
@@ -258,6 +258,22 @@ namespace SlateIM
 		return ensure(CurrentContainer.Widget) ? CurrentContainer.Widget->GetChild(CurrentContainer.LastUsedChildIndex) : nullptr;
 	}
 
+	TSharedPtr<SWidget> FSlateIMManager::GetCurrentChildAsWidget() const
+	{
+		FSlateIMChild CurrentChild = GetCurrentChild();
+		TSharedPtr<SWidget> ChildWidget = CurrentChild.GetWidget();
+
+		if (!ChildWidget)
+		{
+			if (TSharedPtr<ISlateIMChild> Child = CurrentChild.GetChild())
+			{
+				ChildWidget = Child->GetAsWidget();
+			}
+		}
+
+		return ChildWidget;
+	}
+
 	void FSlateIMManager::UpdateCurrentChild(FSlateIMChild Child, const FSlateIMSlotData& AlignmentData)
 	{
 		FContainerNode& CurrentContainer = CurrentRoot->CurrentContainerStack.Top();
diff --git a/Engine/Plugins/Experimental/SlateIM/Source/SlateIM/Private/Misc/SlateIMManager.h b/Engine/Plugins/Experimental/SlateIM/Source/SlateIM/Private/Misc/SlateIMManager.h
--- a/Engine/Plugins/Experimental/SlateIM/Source/SlateIM/Private/Misc/SlateIMManager.h
+++ b/Engine/Plugins/Experimental/SlateIM/Source/SlateIM/Private/Misc/SlateIMManager.h
@@ -142,6 +142,7 @@ namespace SlateIM
 
 		/** Gets the widget which should be at the current spot in the hierarchy or null if it doesn't exist */
 		FSlateIMChild GetCurrentChild() const;
+		TSharedPtr<SWidget> GetCurrentChildAsWidget() const;
 		const FRootNode& GetCurrentRoot() const { return *CurrentRoot; } 
 		FRootNode& GetMutableCurrentRoot() { return *CurrentRoot; }
 
diff --git a/Engine/Plugins/Experimental/SlateIM/Source/SlateIM/Public/SlateIM.h b/Engine/Plugins/Experimental/SlateIM/Source/SlateIM/Public/SlateIM.h
--- a/Engine/Plugins/Experimental/SlateIM/Source/SlateIM/Public/SlateIM.h
+++ b/Engine/Plugins/Experimental/SlateIM/Source/SlateIM/Public/SlateIM.h
@@ -605,4 +605,31 @@ namespace SlateIM
 	 */
 	SLATEIM_API EAppReturnType::Type ModalDialog(EAppMsgType::Type MessageType, const FStringView& DialogText, EAppMsgCategory Category = EAppMsgCategory::Warning, const FStringView& DialogTitle = FStringView());
 #pragma endregion Util
+
+#pragma region Queries
+	/**
+	 * Query whether the last widget is hovered or not
+	 * 
+	 * @return true if the cursor is hovering the last rendered widget, false otherwise
+	 */
+	SLATEIM_API bool IsHovered();
+
+	enum class EFocusDepth : uint8
+	{
+		// Only check if the previous widget is focused
+		SelfOnly,
+
+		// Check if the previous widget or any of its childs widgets are focused (recursively)
+		IncludingDescendants
+	};
+
+	/**
+	 * Query whether the previous widget is focused
+	 * 
+	 * @param Depth (optional) How far to check for focus. Includes child widgets by default.
+	 * @return true if the previous widget has focus (for the specified depth), false otherwise
+	 */
+	SLATEIM_API bool IsFocused(EFocusDepth Depth = EFocusDepth::IncludingDescendants);
+#pragma endregion Queries
 }
+
