SlateIM: Fix certain widget styles not updating when setting a new style after initial creation

#rnx

diff --git a/Engine/Plugins/Experimental/SlateIM/Source/SlateIM/Private/Implementation/SlateIM_Widgets.cpp b/Engine/Plugins/Experimental/SlateIM/Source/SlateIM/Private/Implementation/SlateIM_Widgets.cpp
--- a/Engine/Plugins/Experimental/SlateIM/Source/SlateIM/Private/Implementation/SlateIM_Widgets.cpp
+++ b/Engine/Plugins/Experimental/SlateIM/Source/SlateIM/Private/Implementation/SlateIM_Widgets.cpp
@@ -278,6 +278,7 @@ namespace SlateIM
 		TSharedPtr<SImButton> ButtonWidget = Scope.GetWidget();
 
 		Scope.HashStringView(InText);
+		Scope.HashData(InStyle);
 
 		bool bWasClicked = false;
 		if (!ButtonWidget)
@@ -306,6 +307,10 @@ namespace SlateIM
 			if (Scope.IsDataHashDirty())
 			{
 				ButtonWidget->SetText(InText);
+				if (InStyle)
+				{
+					ButtonWidget->SetButtonStyle(InStyle);
+				}
 			}
 		}
 
@@ -333,6 +338,7 @@ namespace SlateIM
 		bool bValueChanged = false;
 
 		Scope.HashStringView(InText);
+		Scope.HashData(CheckBoxStyle);
 
 		if (!CheckboxWidget)
 		{
@@ -368,6 +374,10 @@ namespace SlateIM
 			if (Scope.IsDataHashDirty())
 			{
 				CheckboxWidget->SetText(FText::FromStringView(InText));
+				if (CheckBoxStyle)
+				{
+					CheckboxWidget->SetStyle(CheckBoxStyle);
+				}
 			}
 		}
 
@@ -471,6 +481,8 @@ namespace SlateIM
 		FWidgetScope<SSlider> Scope(Defaults::Padding, Defaults::HAlign, Defaults::VAlign, Defaults::bAutoSize, Defaults::InputWidgetWidth);
 		TSharedPtr<SSlider> SliderWidget = Scope.GetWidget();
 
+		Scope.HashData(SliderStyle);
+
 		bool bValueChanged = false;
 		if (!SliderWidget)
 		{
@@ -506,6 +518,11 @@ namespace SlateIM
 			{
 				SliderWidget->SetValue(InOutValue);
 			}
+			
+			if (Scope.IsDataHashDirty() && SliderStyle)
+			{
+				SliderWidget->SetStyle(SliderStyle);
+			}
 		}
 
 		return bValueChanged;
@@ -516,6 +533,8 @@ namespace SlateIM
 		FWidgetScope<SProgressBar> Scope(Defaults::Padding, Defaults::HAlign, Defaults::VAlign, Defaults::bAutoSize, Defaults::InputWidgetWidth);
 		TSharedPtr<SProgressBar> ProgressBarWidget = Scope.GetWidget();
 
+		Scope.HashData(ProgressBarStyle);
+
 		if (!ProgressBarWidget)
 		{
 			ProgressBarWidget =
@@ -532,6 +551,11 @@ namespace SlateIM
 		else
 		{
 			ProgressBarWidget->SetPercent(Percent);
+			
+			if (Scope.IsDataHashDirty() && ProgressBarStyle)
+			{
+				ProgressBarWidget->SetStyle(ProgressBarStyle);
+			}
 		}
 	}
 
