SlateIM: Add a graph widget and Line Graphing functions

Serves as a starting point to support plotting more complex graphs. Does not have any user interactivity features yet.

#rnx

diff --git a/Engine/Plugins/Experimental/SlateIM/Source/SlateIM/Private/Implementation/SlateIM_Graphing.cpp b/Engine/Plugins/Experimental/SlateIM/Source/SlateIM/Private/Implementation/SlateIM_Graphing.cpp
new file mode 100644
--- /dev/null
+++ b/Engine/Plugins/Experimental/SlateIM/Source/SlateIM/Private/Implementation/SlateIM_Graphing.cpp
@@ -0,0 +1,64 @@
+﻿// Copyright Epic Games, Inc. All Rights Reserved.
+
+#include "SlateIM.h"
+
+#include "Misc/SlateIMManager.h"
+#include "Misc/SlateIMWidgetScope.h"
+#include "Widgets/SImGraph.h"
+
+namespace SlateIM
+{
+	TSharedPtr<SImGraph> GetCurrentGraphWidget()
+	{
+		if (TSharedPtr<SWidget> ChildWidget = FSlateIMManager::Get().GetCurrentChildAsWidget())
+		{
+			if (ChildWidget->GetWidgetClass().GetWidgetType() == SImGraph::StaticWidgetClass().GetWidgetType())
+			{
+				return StaticCastSharedPtr<SImGraph>(ChildWidget);
+			}
+		}
+
+		return nullptr;
+	}
+	
+	void BeginGraph()
+	{
+		FWidgetScope<SImGraph> Scope;
+		TSharedPtr<SImGraph> GraphWidget = Scope.GetWidget();
+
+		if (!GraphWidget)
+		{
+			GraphWidget = SNew(SImGraph);
+			Scope.UpdateWidget(GraphWidget);
+		}
+		
+		GraphWidget->BeginGraph();
+	}
+
+	void EndGraph()
+	{
+		TSharedPtr<SImGraph> GraphWidget = GetCurrentGraphWidget();
+		if (ensureMsgf(GraphWidget, TEXT("Calling SlateIM::EndGraph() but the last widget is not a graph widget")))
+		{
+			GraphWidget->EndGraph();
+		}
+	}
+
+	void GraphLine(const TArrayView<FVector2D>& Points, const FLinearColor& LineColor, float LineThickness, const FDoubleRange& XViewRange, const FDoubleRange& YViewRange)
+	{
+		TSharedPtr<SImGraph> GraphWidget = GetCurrentGraphWidget();
+		if (ensureMsgf(GraphWidget, TEXT("Calling SlateIM::GraphLine() but the last widget is not a graph widget")))
+		{
+			GraphWidget->AddLineGraph(Points, LineColor, LineThickness, XViewRange, YViewRange);
+		}
+	}
+
+	void GraphLine(const TArrayView<double>& Values, const FLinearColor& LineColor, float LineThickness, const FDoubleRange& ViewRange)
+	{
+		TSharedPtr<SImGraph> GraphWidget = GetCurrentGraphWidget();
+		if (ensureMsgf(GraphWidget, TEXT("Calling SlateIM::GraphLine() but the last widget is not a graph widget")))
+		{
+			GraphWidget->AddLineGraph(Values, LineColor, LineThickness, ViewRange);
+		}
+	}
+}
diff --git a/Engine/Plugins/Experimental/SlateIM/Source/SlateIM/Private/Misc/SlateIMExamples.cpp b/Engine/Plugins/Experimental/SlateIM/Source/SlateIM/Private/Misc/SlateIMExamples.cpp
--- a/Engine/Plugins/Experimental/SlateIM/Source/SlateIM/Private/Misc/SlateIMExamples.cpp
+++ b/Engine/Plugins/Experimental/SlateIM/Source/SlateIM/Private/Misc/SlateIMExamples.cpp
@@ -498,6 +498,74 @@ void FSlateIMTestWidget::Draw()
 		}
 		SlateIM::EndHorizontalStack();
 
+		// Graph Examples
+		{
+			SCOPED_NAMED_EVENT_TEXT("Graph Example", FColorList::Goldenrod);
+			SlateIM::BeginHorizontalStack();
+			{
+				if (SquareGraphValues.Num() >= 100)
+				{
+					SquareGraphValues.PopFront();
+				}
+				const double SquareValue = (GFrameCounter / 10) % 2;
+				SquareGraphValues.Emplace(SquareValue);
+				SlateIM::Fill();
+				SlateIM::MinHeight(200.f);
+				SlateIM::BeginGraph();
+				SlateIM::GraphLine(SquareGraphValues.Compact(), FLinearColor::White, 3.f, FDoubleRange(0, 1.0));
+				SlateIM::EndGraph();
+
+				SlateIM::VAlign(VAlign_Center);
+				SlateIM::Text(FString::Printf(TEXT("%0.3lf"), SquareValue));
+
+				const double NextSinX = SinGraphValues.Last().X + 1.0;
+				const double NextCosX = CosGraphValues.Last().X + 1.0;
+				const double NextTanX = TanGraphValues.Last().X + 1.0;
+				if (GFrameCounter % 4 == 0)
+				{
+					if (SinGraphValues.Num() >= 100)
+					{
+						SinGraphValues.PopFront();
+					}
+					SinGraphValues.Emplace(NextSinX, FMath::Sin(NextSinX/4.0));
+					
+					if (CosGraphValues.Num() >= 100)
+					{
+						CosGraphValues.PopFront();
+					}
+					CosGraphValues.Emplace(NextCosX, FMath::Cos(NextCosX/4.0));
+					
+					if (TanGraphValues.Num() >= 100)
+					{
+						TanGraphValues.PopFront();
+					}
+					TanGraphValues.Emplace(NextTanX, FMath::Tan(NextTanX/4.0));
+				}
+				SlateIM::Fill();
+				SlateIM::MinHeight(200.f);
+				SlateIM::BeginGraph();
+				SlateIM::GraphLine(SinGraphValues.Compact(), FColor::Orange, 1.f, FDoubleRange(NextSinX - 100.0, NextSinX), FDoubleRange(-1.5, 1.5));
+				SlateIM::GraphLine(CosGraphValues.Compact(), FLinearColor::Green, 1.f, FDoubleRange(NextCosX - 100.0, NextCosX), FDoubleRange(-1.5, 1.5));
+				SlateIM::GraphLine(TanGraphValues.Compact(), FColor::Magenta, 1.f, FDoubleRange(NextTanX - 100.0, NextTanX), FDoubleRange(-1.5, 1.5));
+				SlateIM::EndGraph();
+
+				SlateIM::MinWidth(50.f);
+				SlateIM::MaxWidth(50.f);
+				SlateIM::BeginVerticalStack();
+				SlateIM::Fill();
+				SlateIM::VAlign(VAlign_Center);
+				SlateIM::Text(FString::Printf(TEXT("%0.3lf"), SinGraphValues.Last().Y), FColor::Orange);
+				SlateIM::Fill();
+				SlateIM::VAlign(VAlign_Center);
+				SlateIM::Text(FString::Printf(TEXT("%0.3lf"), CosGraphValues.Last().Y), FLinearColor::Green);
+				SlateIM::Fill();
+				SlateIM::VAlign(VAlign_Center);
+				SlateIM::Text(FString::Printf(TEXT("%0.3lf"), TanGraphValues.Last().Y), FColor::Magenta);
+				SlateIM::EndVerticalStack();
+			}
+			SlateIM::EndHorizontalStack();
+		}
+
 		// Tree Example
 		{
 			SCOPED_NAMED_EVENT_TEXT("Tree Example", FColorList::Goldenrod);
diff --git a/Engine/Plugins/Experimental/SlateIM/Source/SlateIM/Private/Misc/SlateIMExamples.h b/Engine/Plugins/Experimental/SlateIM/Source/SlateIM/Private/Misc/SlateIMExamples.h
--- a/Engine/Plugins/Experimental/SlateIM/Source/SlateIM/Private/Misc/SlateIMExamples.h
+++ b/Engine/Plugins/Experimental/SlateIM/Source/SlateIM/Private/Misc/SlateIMExamples.h
@@ -1,6 +1,7 @@
 ﻿// Copyright Epic Games, Inc. All Rights Reserved.
 
 #pragma once
+#include "Containers/RingBuffer.h"
 
 #if WITH_SLATEIM_EXAMPLES
 #include "SlateIMWidgetBase.h"
@@ -35,6 +36,10 @@ private:
 	FString NumItemsText = FString::FromInt(NumItems);
 	bool bShouldLiveUpdateTable = false;
 	const TArray<FString> ComboBoxItems = { TEXT("Option 1"), TEXT("Option 2"), TEXT("Option 3"), TEXT("Option 4") };
+	TRingBuffer<double> SquareGraphValues;
+	TRingBuffer<FVector2D> SinGraphValues = { FVector2D(0, FMath::Sin(0.0)) };
+	TRingBuffer<FVector2D> CosGraphValues = { FVector2D(0, FMath::Cos(0.0)) };
+	TRingBuffer<FVector2D> TanGraphValues = { FVector2D(0, FMath::Tan(0.0)) };
 
 #if WITH_ENGINE
 	TSoftObjectPtr<UTexture2D> RedIcon = TSoftObjectPtr<UTexture2D>(FSoftObjectPath(TEXT("/Engine/EngineResources/AICON-Red.AICON-Red")));
diff --git a/Engine/Plugins/Experimental/SlateIM/Source/SlateIM/Private/Widgets/SImGraph.cpp b/Engine/Plugins/Experimental/SlateIM/Source/SlateIM/Private/Widgets/SImGraph.cpp
new file mode 100644
--- /dev/null
+++ b/Engine/Plugins/Experimental/SlateIM/Source/SlateIM/Private/Widgets/SImGraph.cpp
@@ -0,0 +1,126 @@
+﻿// Copyright Epic Games, Inc. All Rights Reserved.
+
+
+#include "SImGraph.h"
+
+SLATE_IMPLEMENT_WIDGET(SImGraph)
+
+void SImGraph::PrivateRegisterAttributes(FSlateAttributeInitializer& AttributeInitializer)
+{
+}
+
+void SImGraph::Construct(const FArguments& InArgs)
+{
+	SetClipping(EWidgetClipping::ClipToBounds);
+}
+
+void SImGraph::BeginGraph()
+{
+	for (FSlateIMLineGraphData& LineGraph : LineGraphs)
+	{
+		LineGraph.bIsStale = true;
+	}
+}
+
+void SImGraph::EndGraph()
+{
+	const int32 FirstStaleIndex = LineGraphs.IndexOfByPredicate([](const FSlateIMLineGraphData& LineGraph)
+	{
+		return LineGraph.bIsStale;
+	});
+
+	if (FirstStaleIndex != INDEX_NONE)
+	{
+		LineGraphs.SetNum(FirstStaleIndex, EAllowShrinking::No);
+	}
+}
+
+void SImGraph::AddLineGraph(const TArrayView<FVector2D>& Points, const FLinearColor& LineColor, float LineThickness, const FDoubleRange& XViewRange, const FDoubleRange& YViewRange)
+{
+	ensure(XViewRange.HasLowerBound() && XViewRange.HasUpperBound());
+	ensure(YViewRange.HasLowerBound() && YViewRange.HasUpperBound());
+	
+	FSlateIMLineGraphData& LineGraph = GetNextLineGraph(LineColor, LineThickness, XViewRange, YViewRange);
+	LineGraph.NormalizedPoints.Reserve(Points.Num());
+
+	for (const FVector2D& Point : Points)
+	{
+		LineGraph.NormalizedPoints.Emplace(
+			(Point.X - XViewRange.GetLowerBound().GetValue()) / (XViewRange.GetUpperBound().GetValue() - XViewRange.GetLowerBound().GetValue()),
+			(Point.Y - YViewRange.GetLowerBound().GetValue()) / (YViewRange.GetUpperBound().GetValue() - YViewRange.GetLowerBound().GetValue()));
+	}
+}
+
+void SImGraph::AddLineGraph(const TArrayView<double>& Values, const FLinearColor& LineColor, float LineThickness, const FDoubleRange& ViewRange)
+{
+	ensure(ViewRange.HasLowerBound() && ViewRange.HasUpperBound());
+
+	const FDoubleRange XViewRange = FDoubleRange(0, static_cast<double>(Values.Num()));
+	
+	FSlateIMLineGraphData& LineGraph = GetNextLineGraph(LineColor, LineThickness, XViewRange, ViewRange);
+	LineGraph.NormalizedPoints.Reserve(Values.Num());
+
+	double XValue = 0.0;
+	for (double Value : Values)
+	{
+		LineGraph.NormalizedPoints.Emplace(
+			(XValue - XViewRange.GetLowerBound().GetValue()) / (XViewRange.GetUpperBound().GetValue() - XViewRange.GetLowerBound().GetValue()),
+			(Value - ViewRange.GetLowerBound().GetValue()) / (ViewRange.GetUpperBound().GetValue() - ViewRange.GetLowerBound().GetValue()));
+
+		XValue += 1.0;
+	}
+}
+
+int32 SImGraph::OnPaint(const FPaintArgs& Args, const FGeometry& AllottedGeometry, const FSlateRect& MyCullingRect, FSlateWindowElementList& OutDrawElements, int32 LayerId, const FWidgetStyle& InWidgetStyle, bool bParentEnabled) const
+{
+	const FVector2D WidgetSize = AllottedGeometry.GetLocalSize();
+
+	TArray<FVector2D> Points;
+	for (const FSlateIMLineGraphData& LineGraph : LineGraphs)
+	{
+		Points.Reset(LineGraph.NormalizedPoints.Num());
+
+		for (const FVector2D& NormalizedPoint : LineGraph.NormalizedPoints)
+		{
+			Points.Emplace(WidgetSize.X * NormalizedPoint.X, WidgetSize.Y - WidgetSize.Y * NormalizedPoint.Y);
+		}
+		
+		FSlateDrawElement::MakeLines(
+			OutDrawElements,
+			++LayerId,
+			AllottedGeometry.ToPaintGeometry(),
+			Points,
+			ESlateDrawEffect::NoPixelSnapping,
+			LineGraph.Color,
+			true,
+			LineGraph.LineThickness);
+	}
+	
+	return LayerId;
+}
+
+FVector2D SImGraph::ComputeDesiredSize(float) const
+{
+	// Graph size should be set by the SlateIM slot, so just return a reasonable default
+	return FVector2D(300, 200);
+}
+
+FSlateIMLineGraphData& SImGraph::GetNextLineGraph(const FLinearColor& LineColor, float LineThickness, const FDoubleRange& XViewRange, const FDoubleRange& YViewRange)
+{
+	for (FSlateIMLineGraphData& ExistingLineGraph : LineGraphs)
+	{
+		if (ExistingLineGraph.bIsStale)
+		{
+			ExistingLineGraph.NormalizedPoints.Reset();
+			ExistingLineGraph.Color = LineColor;
+			ExistingLineGraph.LineThickness = LineThickness;
+			ExistingLineGraph.XViewRange = XViewRange;
+			ExistingLineGraph.YViewRange = YViewRange;
+			ExistingLineGraph.bIsStale = false;
+				
+			return ExistingLineGraph;
+		}
+	}
+
+	return LineGraphs.Emplace_GetRef(TArray<FVector2D>(), LineColor, LineThickness, XViewRange, YViewRange);
+}
diff --git a/Engine/Plugins/Experimental/SlateIM/Source/SlateIM/Private/Widgets/SImGraph.h b/Engine/Plugins/Experimental/SlateIM/Source/SlateIM/Private/Widgets/SImGraph.h
new file mode 100644
--- /dev/null
+++ b/Engine/Plugins/Experimental/SlateIM/Source/SlateIM/Private/Widgets/SImGraph.h
@@ -0,0 +1,43 @@
+﻿// Copyright Epic Games, Inc. All Rights Reserved.
+
+#pragma once
+
+#include "Widgets/SLeafWidget.h"
+
+struct FSlateIMLineGraphData
+{
+	TArray<FVector2D> NormalizedPoints;
+	FLinearColor Color = FLinearColor::White;
+	float LineThickness = 1.0f;
+	FDoubleRange XViewRange;
+	FDoubleRange YViewRange;
+	bool bIsStale = false;
+};
+
+class SLATEIM_API SImGraph : public SLeafWidget
+{
+	SLATE_DECLARE_WIDGET(SImGraph, SLeafWidget)
+	
+public:
+	SLATE_BEGIN_ARGS(SImGraph)
+		{
+		}
+
+	SLATE_END_ARGS()
+
+	void Construct(const FArguments& InArgs);
+
+	void BeginGraph();
+	void EndGraph();
+
+	void AddLineGraph(const TArrayView<FVector2D>& Points, const FLinearColor& LineColor, float LineThickness, const FDoubleRange& XViewRange, const FDoubleRange& YViewRange);
+	void AddLineGraph(const TArrayView<double>& Points, const FLinearColor& LineColor, float LineThickness, const FDoubleRange& ViewRange);
+
+private:
+	virtual int32 OnPaint( const FPaintArgs& Args, const FGeometry& AllottedGeometry, const FSlateRect& MyCullingRect, FSlateWindowElementList& OutDrawElements, int32 LayerId, const FWidgetStyle& InWidgetStyle, bool bParentEnabled ) const override;
+	virtual FVector2D ComputeDesiredSize(float) const override;
+
+	FSlateIMLineGraphData& GetNextLineGraph(const FLinearColor& LineColor, float LineThickness, const FDoubleRange& XViewRange, const FDoubleRange& YViewRange);
+	
+	TArray<FSlateIMLineGraphData> LineGraphs;
+};
diff --git a/Engine/Plugins/Experimental/SlateIM/Source/SlateIM/Public/SlateIM.h b/Engine/Plugins/Experimental/SlateIM/Source/SlateIM/Public/SlateIM.h
--- a/Engine/Plugins/Experimental/SlateIM/Source/SlateIM/Public/SlateIM.h
+++ b/Engine/Plugins/Experimental/SlateIM/Source/SlateIM/Public/SlateIM.h
@@ -631,5 +631,39 @@ namespace SlateIM
 	 */
 	SLATEIM_API bool IsFocused(EFocusDepth Depth = EFocusDepth::IncludingDescendants);
 #pragma endregion Queries
+
+#pragma region Graphing
+	/**
+	 * Begin a graph widget, call graphing functions between this and EndGraph to include multiple graphs in a single chart
+	 */
+	SLATEIM_API void BeginGraph();
+
+	/**
+	 * Call when finished calling graphing functions for the widget
+	 */
+	SLATEIM_API void EndGraph();
+
+	/**
+	 * Add a line graph of 2D vectors to the graph widget
+	 * 
+	 * @param Points The X,Y points to plot on the graph
+	 * @param LineColor The color of the line to draw for this graph
+	 * @param LineThickness How thick to draw the line
+	 * @param XViewRange The min and max X values to horizontally scale the graph to
+	 * @param YViewRange The min and max Y values to vertically scale the graph to
+	 */
+	SLATEIM_API void GraphLine(const TArrayView<FVector2D>& Points, const FLinearColor& LineColor, float LineThickness, const FDoubleRange& XViewRange, const FDoubleRange& YViewRange);
+
+	/**
+	 * Add a line graph of values to the graph widget.
+	 * The value index is used as the X-value when plotting, the graph will scale horizontally to fit all values in the array.
+	 * 
+	 * @param Values The Y values to plot on the graph
+	 * @param LineColor The color of the line to draw for this graph
+	 * @param LineThickness How thick to draw the line
+	 * @param ViewRange The min and max Y values to vertically scale the graph to
+	 */
+	SLATEIM_API void GraphLine(const TArrayView<double>& Values, const FLinearColor& LineColor, float LineThickness, const FDoubleRange& ViewRange);
+#pragma endregion Graphing
 }
 
