SlateIM: Allow setting layout properties for viewport roots

#rnx

diff --git a/Engine/Plugins/Experimental/SlateIM/Source/SlateIM/Private/Implementation/SlateIM_Roots.cpp b/Engine/Plugins/Experimental/SlateIM/Source/SlateIM/Private/Implementation/SlateIM_Roots.cpp
--- a/Engine/Plugins/Experimental/SlateIM/Source/SlateIM/Private/Implementation/SlateIM_Roots.cpp
+++ b/Engine/Plugins/Experimental/SlateIM/Source/SlateIM/Private/Implementation/SlateIM_Roots.cpp
@@ -33,6 +33,7 @@ namespace SlateIM
 			RootNode->RootState = RootNode->RootWidget->IsVisible();
 		}
 
+		TSharedPtr<FSlateIMWindowRoot> WindowRoot;
 		if (bNeedsCreate)
 		{
 			// TODO - Save and restore window size and position, mapped to UniqueName
@@ -43,24 +44,24 @@ namespace SlateIM
 
 			FSlateApplication::Get().AddWindow(NewWindowWidget);
 
-			TSharedRef<FSlateIMWindowRoot> NewRoot = MakeShared<FSlateIMWindowRoot>(NewWindowWidget);
+			WindowRoot = MakeShared<FSlateIMWindowRoot>(NewWindowWidget);
 
 			if (RootNode)
 			{
-				RootNode->RootWidget = NewRoot;
+				RootNode->RootWidget = WindowRoot;
 				RootNode->RootState = true;
 			}
 			else
 			{
-				RootNode = &FSlateIMManager::Get().AddRoot(UniqueName, NewRoot);
+				RootNode = &FSlateIMManager::Get().AddRoot(UniqueName, WindowRoot);
 			}
-
-			RootNode->RootWidget->UpdateRoot(WindowTitle, FVector2f::ZeroVector, WindowSize);
 		}
-		else if (RootNode->RootState)
+		else
 		{
-			RootNode->RootWidget->UpdateRoot(WindowTitle, FVector2f::ZeroVector, WindowSize);
+			WindowRoot = StaticCastSharedPtr<FSlateIMWindowRoot>(RootNode->RootWidget);
 		}
+		
+		WindowRoot->UpdateWindow(WindowTitle);
 
 		FSlateIMManager::Get().BeginRoot(UniqueName);
 
@@ -68,7 +69,7 @@ namespace SlateIM
 	}
 
 #if WITH_ENGINE
-	bool BeginViewportRoot(FName UniqueName, UGameViewportClient* ViewportClient)
+	bool BeginViewportRoot(FName UniqueName, UGameViewportClient* ViewportClient, const FViewportRootLayout& Layout)
 	{
 		if (!FSlateIMManager::Get().CanUpdateSlateIM())
 		{
@@ -77,29 +78,35 @@ namespace SlateIM
 		
 		bool bNeedsCreate = true;
 
-		FRootNode* ViewportRoot = FSlateIMManager::Get().FindRoot<FSlateIMViewportRoot>(UniqueName);
-		if (ViewportRoot)
+		FRootNode* RootNode = FSlateIMManager::Get().FindRoot<FSlateIMViewportRoot>(UniqueName);
+		if (RootNode)
 		{
-			const bool bIsValidRoot = ViewportRoot->RootWidget.IsValid() && StaticCastSharedPtr<FSlateIMViewportRoot>(ViewportRoot->RootWidget)->GameViewport == ViewportClient;
+			const bool bIsValidRoot = RootNode->RootWidget.IsValid() && StaticCastSharedPtr<FSlateIMViewportRoot>(RootNode->RootWidget)->GameViewport == ViewportClient;
 			bNeedsCreate = !bIsValidRoot;
 
-			ViewportRoot->RootState = bIsValidRoot;
+			RootNode->RootState = bIsValidRoot;
 		}
 
+		TSharedPtr<FSlateIMViewportRoot> ViewportRoot;
 		if (bNeedsCreate)
 		{
-			TSharedRef<FSlateIMViewportRoot> NewRoot = MakeShared<FSlateIMViewportRoot>(ViewportClient);
+			ViewportRoot = MakeShared<FSlateIMViewportRoot>(ViewportClient);
 
-			NewRoot->UpdateRoot(FStringView(), FVector2f::ZeroVector, FVector2f::ZeroVector);
-			ViewportRoot = &FSlateIMManager::Get().AddRoot(UniqueName, NewRoot);
+			RootNode = &FSlateIMManager::Get().AddRoot(UniqueName, ViewportRoot);
+		}
+		else
+		{
+			ViewportRoot = StaticCastSharedPtr<FSlateIMViewportRoot>(RootNode->RootWidget);
 		}
 
+		ViewportRoot->UpdateViewport(Layout);
+
 		FSlateIMManager::Get().BeginRoot(UniqueName);
 
-		return ViewportRoot->RootState;
+		return RootNode->RootState;
 	}
 
-	bool BeginViewportRoot(FName UniqueName, ULocalPlayer* LocalPlayer)
+	bool BeginViewportRoot(FName UniqueName, ULocalPlayer* LocalPlayer, const FViewportRootLayout& Layout)
 	{
 		if (!FSlateIMManager::Get().CanUpdateSlateIM())
 		{
@@ -108,30 +115,36 @@ namespace SlateIM
 		
 		bool bNeedsCreate = true;
 
-		FRootNode* ViewportRoot = FSlateIMManager::Get().FindRoot<FSlateIMViewportRoot>(UniqueName);
-		if (ViewportRoot)
+		FRootNode* RootNode = FSlateIMManager::Get().FindRoot<FSlateIMViewportRoot>(UniqueName);
+		if (RootNode)
 		{
-			const bool bIsValidRoot = ViewportRoot->RootWidget.IsValid() && StaticCastSharedPtr<FSlateIMViewportRoot>(ViewportRoot->RootWidget)->LocalPlayer == LocalPlayer;
+			const bool bIsValidRoot = RootNode->RootWidget.IsValid() && StaticCastSharedPtr<FSlateIMViewportRoot>(RootNode->RootWidget)->LocalPlayer == LocalPlayer;
 			bNeedsCreate = !bIsValidRoot;
 
-			ViewportRoot->RootState = bIsValidRoot;
+			RootNode->RootState = bIsValidRoot;
 		}
 
+		TSharedPtr<FSlateIMViewportRoot> ViewportRoot;
 		if (bNeedsCreate)
 		{
-			TSharedRef<FSlateIMViewportRoot> NewRoot = MakeShared<FSlateIMViewportRoot>(LocalPlayer);
+			ViewportRoot = MakeShared<FSlateIMViewportRoot>(LocalPlayer);
 
-			NewRoot->UpdateRoot(FStringView(), FVector2f::ZeroVector, FVector2f::ZeroVector);
-			ViewportRoot = &FSlateIMManager::Get().AddRoot(UniqueName, NewRoot);
+			RootNode = &FSlateIMManager::Get().AddRoot(UniqueName, ViewportRoot);
 		}
+		else
+		{
+			ViewportRoot = StaticCastSharedPtr<FSlateIMViewportRoot>(RootNode->RootWidget);
+		}
+
+		ViewportRoot->UpdateViewport(Layout);
 
 		FSlateIMManager::Get().BeginRoot(UniqueName);
 
-		return ViewportRoot->RootState;
+		return RootNode->RootState;
 	}
 
 #if WITH_EDITOR
-	bool BeginViewportRoot(FName UniqueName, TSharedPtr<IAssetViewport> AssetViewport, FVector2f Offset, FVector2f NormalizedSize)
+	bool BeginViewportRoot(FName UniqueName, TSharedPtr<IAssetViewport> AssetViewport, const FViewportRootLayout& Layout)
 	{
 		if (!FSlateIMManager::Get().CanUpdateSlateIM())
 		{
@@ -140,30 +153,32 @@ namespace SlateIM
 		
 		bool bNeedsCreate = true;
 
-		FRootNode* ViewportRoot = FSlateIMManager::Get().FindRoot<FSlateIMViewportRoot>(UniqueName);
-		if (ViewportRoot)
+		FRootNode* RootNode = FSlateIMManager::Get().FindRoot<FSlateIMViewportRoot>(UniqueName);
+		if (RootNode)
 		{
-			const bool bIsValidRoot = ViewportRoot->RootWidget.IsValid() && StaticCastSharedPtr<FSlateIMViewportRoot>(ViewportRoot->RootWidget)->AssetViewport == AssetViewport;
+			const bool bIsValidRoot = RootNode->RootWidget.IsValid() && StaticCastSharedPtr<FSlateIMViewportRoot>(RootNode->RootWidget)->AssetViewport == AssetViewport;
 			bNeedsCreate = !bIsValidRoot;
 
-			ViewportRoot->RootState = bIsValidRoot;
+			RootNode->RootState = bIsValidRoot;
 		}
-
+		
+		TSharedPtr<FSlateIMViewportRoot> ViewportRoot;
 		if (bNeedsCreate)
 		{
-			TSharedRef<FSlateIMViewportRoot> NewRoot = MakeShared<FSlateIMViewportRoot>(AssetViewport);
+			ViewportRoot = MakeShared<FSlateIMViewportRoot>(AssetViewport);
 
-			NewRoot->UpdateRoot(FStringView(), Offset, NormalizedSize);
-			ViewportRoot = &FSlateIMManager::Get().AddRoot(UniqueName, NewRoot);
+			RootNode = &FSlateIMManager::Get().AddRoot(UniqueName, ViewportRoot);
 		}
-		else if (ViewportRoot->RootState)
+		else
 		{
-			ViewportRoot->RootWidget->UpdateRoot(FStringView(), Offset, NormalizedSize);
+			ViewportRoot = StaticCastSharedPtr<FSlateIMViewportRoot>(RootNode->RootWidget);
 		}
 
+		ViewportRoot->UpdateViewport(Layout);
+
 		FSlateIMManager::Get().BeginRoot(UniqueName);
 
-		return ViewportRoot->RootState;
+		return RootNode->RootState;
 	}
 #endif // WITH_EDITOR
 
@@ -188,7 +203,6 @@ namespace SlateIM
 		if (bNeedsCreate)
 		{
 			TSharedRef<FSlateIMExposedRoot> NewRoot = MakeShared<FSlateIMExposedRoot>();
-			NewRoot->UpdateRoot(FStringView(), FVector2f::ZeroVector, FVector2f::UnitVector);
 			ExistingRoot = &FSlateIMManager::Get().AddRoot(UniqueName, NewRoot);
 			OutSlateIMWidget = NewRoot->GetExposedWidget();
 		}
diff --git a/Engine/Plugins/Experimental/SlateIM/Source/SlateIM/Private/Misc/SlateIMExamples.cpp b/Engine/Plugins/Experimental/SlateIM/Source/SlateIM/Private/Misc/SlateIMExamples.cpp
--- a/Engine/Plugins/Experimental/SlateIM/Source/SlateIM/Private/Misc/SlateIMExamples.cpp
+++ b/Engine/Plugins/Experimental/SlateIM/Source/SlateIM/Private/Misc/SlateIMExamples.cpp
@@ -1090,7 +1090,7 @@ void FSlateIMTestViewportWidget::DrawWidget(float DeltaTime)
 {
 	if (GEngine && GEngine->GameViewport)
 	{
-		if (SlateIM::BeginViewportRoot("SlateIMTestSuiteViewport", GEngine->GameViewport))
+		if (SlateIM::BeginViewportRoot("SlateIMTestSuiteViewport", GEngine->GameViewport, Layout))
 		{
 			TestWidget.Draw();
 		}
@@ -1102,7 +1102,7 @@ void FSlateIMTestViewportWidget::DrawWidget(float DeltaTime)
 		TSharedPtr<SLevelViewport> LevelViewport = StaticCastSharedPtr<SLevelViewport>(GCurrentLevelEditingViewportClient->GetEditorViewportWidget());
 		if (LevelViewport.IsValid())
 		{
-			if (SlateIM::BeginViewportRoot("SlateIMTestSuiteViewport", LevelViewport, FVector2f(0, 25), FVector2f(.5f, .75f)))
+			if (SlateIM::BeginViewportRoot("SlateIMTestSuiteViewport", LevelViewport, Layout))
 			{
 				TestWidget.Draw();
 			}
diff --git a/Engine/Plugins/Experimental/SlateIM/Source/SlateIM/Private/Misc/SlateIMExamples.h b/Engine/Plugins/Experimental/SlateIM/Source/SlateIM/Private/Misc/SlateIMExamples.h
--- a/Engine/Plugins/Experimental/SlateIM/Source/SlateIM/Private/Misc/SlateIMExamples.h
+++ b/Engine/Plugins/Experimental/SlateIM/Source/SlateIM/Private/Misc/SlateIMExamples.h
@@ -4,6 +4,7 @@
 #include "Containers/RingBuffer.h"
 
 #if WITH_SLATEIM_EXAMPLES
+#include "SlateIM.h"
 #include "SlateIMWidgetBase.h"
 #include "Styling/SlateTypes.h"
 
@@ -84,12 +85,14 @@ class FSlateIMTestViewportWidget : public FSlateIMWidgetWithCommandBase
 public:
 	FSlateIMTestViewportWidget(const TCHAR* Command, const TCHAR* CommandHelp)
 		: FSlateIMWidgetWithCommandBase(Command, CommandHelp)
+		, Layout(FAnchors(0.5f, 0), FVector2f::ZeroVector, FVector2f(0.5f, 0), FVector2f(700, 900))
 	{}
 	
 	virtual void DrawWidget(float DeltaTime) override;
 
 private:
 	FSlateIMTestWidget TestWidget;
+	SlateIM::FViewportRootLayout Layout;
 };
 #endif // WITH_ENGINE
 #endif // WITH_SLATEIM_EXAMPLES
diff --git a/Engine/Plugins/Experimental/SlateIM/Source/SlateIM/Private/Roots/ISlateIMRoot.h b/Engine/Plugins/Experimental/SlateIM/Source/SlateIM/Private/Roots/ISlateIMRoot.h
--- a/Engine/Plugins/Experimental/SlateIM/Source/SlateIM/Private/Roots/ISlateIMRoot.h
+++ b/Engine/Plugins/Experimental/SlateIM/Source/SlateIM/Private/Roots/ISlateIMRoot.h
@@ -13,7 +13,6 @@ class SWidget;
 class ISlateIMRoot : public ISlateIMTypeChecking
 {
 public:
-	virtual void UpdateRoot(const FStringView& Title, FVector2f Offset, FVector2f Size) = 0;
 	virtual void UpdateChild(TSharedRef<SWidget> Child, const FSlateIMSlotData& AlignmentData) = 0;
 	virtual bool IsVisible() const = 0;
 };
diff --git a/Engine/Plugins/Experimental/SlateIM/Source/SlateIM/Private/Roots/SlateIMExposedRoot.cpp b/Engine/Plugins/Experimental/SlateIM/Source/SlateIM/Private/Roots/SlateIMExposedRoot.cpp
--- a/Engine/Plugins/Experimental/SlateIM/Source/SlateIM/Private/Roots/SlateIMExposedRoot.cpp
+++ b/Engine/Plugins/Experimental/SlateIM/Source/SlateIM/Private/Roots/SlateIMExposedRoot.cpp
@@ -7,23 +7,13 @@
 
 
 FSlateIMExposedRoot::FSlateIMExposedRoot()
+	: ExposedWidget(SNew(SBox))
 {
 }
 
 FSlateIMExposedRoot::~FSlateIMExposedRoot()
 {
-	if (ExposedWidget.IsValid())
-	{
-		ExposedWidget->SetContent(SNullWidget::NullWidget);
-	}
-}
-
-void FSlateIMExposedRoot::UpdateRoot(const FStringView& Title, FVector2f Offset, FVector2f Size)
-{
-	if (!ExposedWidget.IsValid())
-	{
-		SAssignNew(ExposedWidget, SBox);
-	}
+	ExposedWidget->SetContent(SNullWidget::NullWidget);
 }
 
 void FSlateIMExposedRoot::UpdateChild(TSharedRef<SWidget> Child, const FSlateIMSlotData& AlignmentData)
@@ -36,5 +26,5 @@ void FSlateIMExposedRoot::UpdateChild(TSharedRef<SWidget> Child, const FSlateIMS
 
 TSharedRef<SWidget> FSlateIMExposedRoot::GetExposedWidget() const
 {
-	return ExposedWidget ? ExposedWidget.ToSharedRef() : SNullWidget::NullWidget;
+	return ExposedWidget;
 }
diff --git a/Engine/Plugins/Experimental/SlateIM/Source/SlateIM/Private/Roots/SlateIMExposedRoot.h b/Engine/Plugins/Experimental/SlateIM/Source/SlateIM/Private/Roots/SlateIMExposedRoot.h
--- a/Engine/Plugins/Experimental/SlateIM/Source/SlateIM/Private/Roots/SlateIMExposedRoot.h
+++ b/Engine/Plugins/Experimental/SlateIM/Source/SlateIM/Private/Roots/SlateIMExposedRoot.h
@@ -14,13 +14,12 @@ public:
 	FSlateIMExposedRoot();
 	virtual ~FSlateIMExposedRoot() override;
 
-	virtual void UpdateRoot(const FStringView& Title, FVector2f Offset, FVector2f Size) override;
 	virtual void UpdateChild(TSharedRef<SWidget> Child, const FSlateIMSlotData& AlignmentData) override;
 	virtual bool IsVisible() const override { return true; }
-
+	
 	TSharedRef<SWidget> GetExposedWidget() const;
 	
 private:
-	TSharedPtr<SBox> ExposedWidget;
+	TSharedRef<SBox> ExposedWidget;
 
 };
diff --git a/Engine/Plugins/Experimental/SlateIM/Source/SlateIM/Private/Roots/SlateIMViewportRoot.cpp b/Engine/Plugins/Experimental/SlateIM/Source/SlateIM/Private/Roots/SlateIMViewportRoot.cpp
--- a/Engine/Plugins/Experimental/SlateIM/Source/SlateIM/Private/Roots/SlateIMViewportRoot.cpp
+++ b/Engine/Plugins/Experimental/SlateIM/Source/SlateIM/Private/Roots/SlateIMViewportRoot.cpp
@@ -1,12 +1,13 @@
 ﻿// Copyright Epic Games, Inc. All Rights Reserved.
 
 
+#if WITH_ENGINE
 #include "SlateIMViewportRoot.h"
 
 #if WITH_EDITOR
 #include "IAssetViewport.h"
 #endif
-#if WITH_ENGINE
+#include "SlateIM.h"
 #include "Engine/GameViewportClient.h"
 #include "Engine/LocalPlayer.h"
 #endif
@@ -60,19 +61,25 @@ FSlateIMViewportRoot::~FSlateIMViewportRoot()
 	}
 }
 
-void FSlateIMViewportRoot::UpdateRoot(const FStringView& Title, FVector2f Offset, FVector2f Size)
+void FSlateIMViewportRoot::UpdateChild(TSharedRef<SWidget> Child, const FSlateIMSlotData& AlignmentData)
+{
+	if (Slot)
+	{
+		Slot->AttachWidget(Child);
+	}
+}
+
+void FSlateIMViewportRoot::UpdateViewport(const SlateIM::FViewportRootLayout& Layout)
 {
 	if (!ViewportCanvas)
 	{
 		ViewportCanvas =
 			SNew(SConstraintCanvas)
 			+ SConstraintCanvas::Slot()
-			.Offset(Offset)
-			.AutoSize(true)
-			.Anchors(FAnchors(0.f, 0.f, Size.X, Size.Y))
-			.Alignment(FVector2D(.0f, .0f))
+			.Anchors(Layout.Anchors)
+			.Alignment(Layout.Alignment)
 			.Expose(Slot);
-
+		
 #if WITH_EDITOR
 		if (TSharedPtr<IAssetViewport> AssetViewportPinned = AssetViewport.Pin())
 		{
@@ -95,16 +102,19 @@ void FSlateIMViewportRoot::UpdateRoot(const FStringView& Title, FVector2f Offset
 	}
 	else
 	{
-		Slot->SetOffset(Offset);
-		Slot->SetAnchors(FAnchors(0.f, 0.f, Size.X, Size.Y));
+		Slot->SetAnchors(Layout.Anchors);
+		Slot->SetAlignment(Layout.Alignment);
 	}
-}
-
-void FSlateIMViewportRoot::UpdateChild(TSharedRef<SWidget> Child, const FSlateIMSlotData& AlignmentData)
-{
-	if (Slot)
+	
+	if (Layout.Size.IsSet())
 	{
-		Slot->AttachWidget(Child);
+		Slot->SetAutoSize(false);
+		Slot->SetOffset(FMargin(Layout.Offset.X, Layout.Offset.Y, Layout.Size->X, Layout.Size->Y));
+	}
+	else
+	{
+		Slot->SetAutoSize(true);
+		Slot->SetOffset(Layout.Offset);
 	}
 }
 #endif // WITH_ENGINE
diff --git a/Engine/Plugins/Experimental/SlateIM/Source/SlateIM/Private/Roots/SlateIMViewportRoot.h b/Engine/Plugins/Experimental/SlateIM/Source/SlateIM/Private/Roots/SlateIMViewportRoot.h
--- a/Engine/Plugins/Experimental/SlateIM/Source/SlateIM/Private/Roots/SlateIMViewportRoot.h
+++ b/Engine/Plugins/Experimental/SlateIM/Source/SlateIM/Private/Roots/SlateIMViewportRoot.h
@@ -7,6 +7,11 @@
 #include "UObject/WeakObjectPtrTemplates.h"
 #include "Widgets/Layout/SConstraintCanvas.h"
 
+namespace SlateIM
+{
+	struct FViewportRootLayout;
+}
+
 class UGameViewportClient;
 class ULocalPlayer;
 
@@ -26,13 +31,14 @@ public:
 #endif
 	virtual ~FSlateIMViewportRoot() override;
 
-	virtual void UpdateRoot(const FStringView& Title, FVector2f Offset, FVector2f Size) override;
 	virtual void UpdateChild(TSharedRef<SWidget> Child, const FSlateIMSlotData& AlignmentData) override;
 	virtual bool IsVisible() const override
 	{
 		return true;
 	}
 
+	void UpdateViewport(const SlateIM::FViewportRootLayout& Layout);
+
 	TWeakObjectPtr<UGameViewportClient> GameViewport;
 	TWeakObjectPtr<ULocalPlayer> LocalPlayer;
 #if WITH_EDITOR
diff --git a/Engine/Plugins/Experimental/SlateIM/Source/SlateIM/Private/Roots/SlateIMWindowRoot.cpp b/Engine/Plugins/Experimental/SlateIM/Source/SlateIM/Private/Roots/SlateIMWindowRoot.cpp
--- a/Engine/Plugins/Experimental/SlateIM/Source/SlateIM/Private/Roots/SlateIMWindowRoot.cpp
+++ b/Engine/Plugins/Experimental/SlateIM/Source/SlateIM/Private/Roots/SlateIMWindowRoot.cpp
@@ -18,18 +18,6 @@ FSlateIMWindowRoot::~FSlateIMWindowRoot()
 	}
 }
 
-void FSlateIMWindowRoot::UpdateRoot(const FStringView& Title, FVector2f Offset, FVector2f Size)
-{
-	if (TSharedPtr<SWindow> RootWindowPin = RootWindow.Pin())
-	{
-		FText NewWindowTitle = FText::FromStringView(Title);
-		if (!RootWindowPin->GetTitle().IdenticalTo(NewWindowTitle))
-		{
-			RootWindowPin->SetTitle(NewWindowTitle);
-		}
-	}
-}
-
 void FSlateIMWindowRoot::UpdateChild(TSharedRef<SWidget> Child, const FSlateIMSlotData& AlignmentData)
 {
 	check(RootWindow != Child);
@@ -41,3 +29,15 @@ bool FSlateIMWindowRoot::IsVisible() const
 {
 	return RootWindow.IsValid();
 }
+
+void FSlateIMWindowRoot::UpdateWindow(const FStringView& Title)
+{
+	if (TSharedPtr<SWindow> RootWindowPin = RootWindow.Pin())
+	{
+		FText NewWindowTitle = FText::FromStringView(Title);
+		if (!RootWindowPin->GetTitle().IdenticalTo(NewWindowTitle))
+		{
+			RootWindowPin->SetTitle(NewWindowTitle);
+		}
+	}
+}
diff --git a/Engine/Plugins/Experimental/SlateIM/Source/SlateIM/Private/Roots/SlateIMWindowRoot.h b/Engine/Plugins/Experimental/SlateIM/Source/SlateIM/Private/Roots/SlateIMWindowRoot.h
--- a/Engine/Plugins/Experimental/SlateIM/Source/SlateIM/Private/Roots/SlateIMWindowRoot.h
+++ b/Engine/Plugins/Experimental/SlateIM/Source/SlateIM/Private/Roots/SlateIMWindowRoot.h
@@ -14,9 +14,10 @@ public:
 	FSlateIMWindowRoot(TSharedRef<SWindow> Window);
 	virtual ~FSlateIMWindowRoot() override;
 
-	virtual void UpdateRoot(const FStringView& Title, FVector2f Offset, FVector2f Size) override;
 	virtual void UpdateChild(TSharedRef<SWidget> Child, const FSlateIMSlotData& AlignmentData) override;
 	virtual bool IsVisible() const override;
+
+	void UpdateWindow(const FStringView& Title);
 	
 private:
 	TWeakPtr<SWindow> RootWindow;
diff --git a/Engine/Plugins/Experimental/SlateIM/Source/SlateIM/Public/SlateIM.h b/Engine/Plugins/Experimental/SlateIM/Source/SlateIM/Public/SlateIM.h
--- a/Engine/Plugins/Experimental/SlateIM/Source/SlateIM/Public/SlateIM.h
+++ b/Engine/Plugins/Experimental/SlateIM/Source/SlateIM/Public/SlateIM.h
@@ -3,10 +3,11 @@
 #pragma once
 
 
+#include "Layout/Margin.h"
 #include "Modules/ModuleInterface.h"
 #include "Styling/SlateColor.h"
 #include "Types/SlateEnums.h"
-#include "Layout/Margin.h"
+#include "Widgets/Layout/Anchors.h"
 
 #if WITH_EDITOR
 class IAssetViewport;
@@ -48,12 +49,29 @@ public:
 namespace SlateIM
 {
 #pragma region Roots
+#if WITH_EDITOR || WITH_ENGINE
+	struct FViewportRootLayout
+	{
+		/** How to anchor the root within the viewport (normalized units) */
+		FAnchors Anchors = FAnchors(0);
+		
+		/** Offset the root from the anchored position in the viewport (in slate units) */
+		FVector2f Offset = FVector2f::ZeroVector;
+
+		/** How the root is aligned to its anchor (normalized units) */
+		FVector2f Alignment = FVector2f::ZeroVector;
+
+		/** Optional set the size (in slate units) of the root in the viewport. When unset, the root will auto-size for its content. */
+		TOptional<FVector2f> Size;
+	};
+#endif // WITH_EDITOR || WITH_ENGINE
+
 	/**
 	 * Begins a new floating window root
 	 *
 	 * @param UniqueName a unique name to identify this window
 	 * @param WindowTitle The name of the window. This will also be displayed in the windows title bar
-	 * @param WindowSize The size to open the window at
+	 * @param WindowSize The size to open the window at (does not update the size of an existing window)
 	 * @param bShouldReopen (optional) If this window was created with a previous call to BeginWindow and then closed, passing in true will reopen the window and false will leave it closed. If the window has never been seen this parameter does nothing and the window will open
 	 *
 	 * @return the current window state. True if open, false if closed or not updating
@@ -69,20 +87,22 @@ namespace SlateIM
 	 * 
 	 * @param UniqueName a globally unique name to identify this root
 	 * @param ViewportClient the viewport client to add the widget to
+	 * @param Layout How to lay out the root within the viewport
 	 * 
 	 * @return Whether the root is valid and updating
 	 */
-	SLATEIM_API bool BeginViewportRoot(FName UniqueName, UGameViewportClient* ViewportClient);
+	SLATEIM_API bool BeginViewportRoot(FName UniqueName, UGameViewportClient* ViewportClient, const FViewportRootLayout& Layout = FViewportRootLayout());
 	
 	/**
 	 *	Begins a new root in the provided player's viewport
 	 * 
 	 * @param UniqueName a globally unique name to identify this root
 	 * @param LocalPlayer the player whose viewport to add the widget to
+	 * @param Layout How to lay out the root within the viewport
 	 * 
 	 * @return Whether the root is valid and updating
 	 */
-	SLATEIM_API bool BeginViewportRoot(FName UniqueName, ULocalPlayer* LocalPlayer);
+	SLATEIM_API bool BeginViewportRoot(FName UniqueName, ULocalPlayer* LocalPlayer, const FViewportRootLayout& Layout = FViewportRootLayout());
 #endif
 
 #if WITH_EDITOR
@@ -91,12 +111,11 @@ namespace SlateIM
 	 * 
 	 * @param UniqueName Unique identifier for this root. Must be globally unique within the same editor session
 	 * @param AssetViewport The editor viewpor to add the SlateIM content to
-	 * @param Offset Offset from the start top left of the viewport corner. Not normalized. 
-	 * @param NormalizedSize Size of the content relative to the size of the viewport in normalized coordinates. I.E as a percentage of viewport dimensions: 0.5 means half the viewport size, 1.0 means full viewport size
+	 * @param Layout How to lay out the root within the viewport
 	 * 
 	 * @return Whether the root is valid and updating
 	 */
-	SLATEIM_API bool BeginViewportRoot(FName UniqueName, TSharedPtr<IAssetViewport> AssetViewport, FVector2f Offset, FVector2f NormalizedSize);
+	SLATEIM_API bool BeginViewportRoot(FName UniqueName, TSharedPtr<IAssetViewport> AssetViewport, const FViewportRootLayout& Layout = FViewportRootLayout());
 #endif
 
 	/**
