SlateIM: Fix Mac Compilation

#rnx

diff --git a/Engine/Plugins/Experimental/SlateIM/Source/SlateIM/Private/Misc/SlateIMExamples.h b/Engine/Plugins/Experimental/SlateIM/Source/SlateIM/Private/Misc/SlateIMExamples.h
--- a/Engine/Plugins/Experimental/SlateIM/Source/SlateIM/Private/Misc/SlateIMExamples.h
+++ b/Engine/Plugins/Experimental/SlateIM/Source/SlateIM/Private/Misc/SlateIMExamples.h
@@ -85,8 +85,11 @@ class FSlateIMTestViewportWidget : public FSlateIMWidgetWithCommandBase
 public:
 	FSlateIMTestViewportWidget(const TCHAR* Command, const TCHAR* CommandHelp)
 		: FSlateIMWidgetWithCommandBase(Command, CommandHelp)
-		, Layout(FAnchors(0.5f, 0), FVector2f::ZeroVector, FVector2f(0.5f, 0), FVector2f(700, 900))
-	{}
+	{
+		Layout.Anchors = FAnchors(0.5f, 0);
+		Layout.Alignment = FVector2f(0.5f, 0);
+		Layout.Size = FVector2f(700, 900);
+	}
 	
 	virtual void DrawWidget(float DeltaTime) override;
 
diff --git a/Engine/Plugins/Experimental/SlateIM/Source/SlateIM/Private/Widgets/SImGraph.cpp b/Engine/Plugins/Experimental/SlateIM/Source/SlateIM/Private/Widgets/SImGraph.cpp
--- a/Engine/Plugins/Experimental/SlateIM/Source/SlateIM/Private/Widgets/SImGraph.cpp
+++ b/Engine/Plugins/Experimental/SlateIM/Source/SlateIM/Private/Widgets/SImGraph.cpp
@@ -122,5 +122,10 @@ FSlateIMLineGraphData& SImGraph::GetNextLineGraph(const FLinearColor& LineColor,
 		}
 	}
 
-	return LineGraphs.Emplace_GetRef(TArray<FVector2D>(), LineColor, LineThickness, XViewRange, YViewRange);
+	FSlateIMLineGraphData NewLineGraph;
+	NewLineGraph.Color = LineColor;
+	NewLineGraph.LineThickness = LineThickness;
+	NewLineGraph.XViewRange = XViewRange;
+	NewLineGraph.YViewRange = YViewRange;
+	return LineGraphs.Emplace_GetRef(MoveTemp(NewLineGraph));
 }
