SlateIM: Add support for changing the options in a ComboBox by setting bForceRefresh = true

This follows the same pattern as SelectionList.

#rnx

diff --git a/Engine/Plugins/Experimental/SlateIM/Source/SlateIM/Private/Implementation/SlateIM_Widgets.cpp b/Engine/Plugins/Experimental/SlateIM/Source/SlateIM/Private/Implementation/SlateIM_Widgets.cpp
--- a/Engine/Plugins/Experimental/SlateIM/Source/SlateIM/Private/Implementation/SlateIM_Widgets.cpp
+++ b/Engine/Plugins/Experimental/SlateIM/Source/SlateIM/Private/Implementation/SlateIM_Widgets.cpp
@@ -559,7 +559,7 @@ namespace SlateIM
 		}
 	}
 
-	bool ComboBox(const TArray<FString>& ComboItems, int32& InOutSelectedItemIndex, const FComboBoxStyle* InComboStyle)
+	bool ComboBox(const TArray<FString>& ComboItems, int32& InOutSelectedItemIndex, bool bForceRefresh, const FComboBoxStyle* InComboStyle)
 	{
 		FWidgetScope<STextComboBox> Scope;
 		TSharedPtr<STextComboBox> ComboWidget = Scope.GetWidget();
@@ -567,19 +567,23 @@ namespace SlateIM
 		Scope.HashData(InComboStyle);
 		
 		using FComboBoxData = FSlateIMDataStore<TArray<TSharedPtr<FString>>>;
-
-		bool bValueChanged = false;
-		if (!ComboWidget)
+		
+		auto BuildComboOptions = [&]()
 		{
 			TArray<TSharedPtr<FString>> Options;
 			Options.Reserve(ComboItems.Num());
 
-			for(const FString& Item : ComboItems)
+			for (const FString& Item : ComboItems)
 			{
 				Options.Add(MakeShared<FString>(Item));
 			}
-		
-			TSharedRef<FComboBoxData> WidgetIMData = MakeShared<FComboBoxData>(MoveTemp(Options));
+			return Options;
+		};
+
+		bool bValueChanged = false;
+		if (!ComboWidget)
+		{
+			TSharedRef<FComboBoxData> WidgetIMData = MakeShared<FComboBoxData>(BuildComboOptions());
 
 			ComboWidget =
 				SNew(STextComboBox)
@@ -601,16 +605,23 @@ namespace SlateIM
 		else
 		{
 			TSharedPtr<FComboBoxData> ComboBoxData = ComboWidget->GetMetaData<FComboBoxData>();
-			bValueChanged = FSlateIMManager::Get().IsWidgetActivatedThisFrame();
-			if (bValueChanged)
+			if (bForceRefresh)
 			{
-				InOutSelectedItemIndex = ComboBoxData->Data.IndexOfByKey(ComboWidget->GetSelectedItem());
+				ComboBoxData->Data = BuildComboOptions();
+				if (InOutSelectedItemIndex != INDEX_NONE && !ComboWidget->IsOpen() && ComboBoxData->Data.IsValidIndex(InOutSelectedItemIndex))
+				{
+					ComboWidget->SetSelectedItem(ComboBoxData->Data[InOutSelectedItemIndex]);
+				}
 			}
 			else
 			{
-				if (InOutSelectedItemIndex != INDEX_NONE && !ComboWidget->IsOpen())
+				bValueChanged = FSlateIMManager::Get().IsWidgetActivatedThisFrame();
+				if (bValueChanged)
 				{
-					ComboWidget->SetSelectedItem(ComboBoxData->Data[InOutSelectedItemIndex]);
+					TSharedPtr<FString> SelectedItem = ComboWidget->GetSelectedItem();
+					const int32 NewSelectedIndex = SelectedItem.IsValid() ? ComboBoxData->Data.IndexOfByKey(SelectedItem) : INDEX_NONE;
+					UE_LOG(LogSlateIM, Verbose, TEXT("Combo Selection Changed %d -> %d"), InOutSelectedItemIndex, NewSelectedIndex);
+					InOutSelectedItemIndex = NewSelectedIndex;
 				}
 			}
 
@@ -699,7 +710,7 @@ namespace SlateIM
 				{
 					TArray<TSharedPtr<FString>> SelectedItems = ListWidget->GetSelectedItems();
 					const int32 NewSelectedIndex = SelectedItems.Num() > 0 ? ListViewData->Data.IndexOfByKey(SelectedItems[0]) : INDEX_NONE;
-					UE_LOG(LogSlateIM, Verbose, TEXT("Selection Changed %d -> %d"), InOutSelectedItemIndex, NewSelectedIndex);
+					UE_LOG(LogSlateIM, Verbose, TEXT("List Selection Changed %d -> %d"), InOutSelectedItemIndex, NewSelectedIndex);
 					InOutSelectedItemIndex = NewSelectedIndex;
 				}
 			}
diff --git a/Engine/Plugins/Experimental/SlateIM/Source/SlateIM/Private/Misc/SlateIMExamples.cpp b/Engine/Plugins/Experimental/SlateIM/Source/SlateIM/Private/Misc/SlateIMExamples.cpp
--- a/Engine/Plugins/Experimental/SlateIM/Source/SlateIM/Private/Misc/SlateIMExamples.cpp
+++ b/Engine/Plugins/Experimental/SlateIM/Source/SlateIM/Private/Misc/SlateIMExamples.cpp
@@ -125,16 +125,34 @@ void FSlateIMTestWidget::Draw()
 			{
 				SlateIM::BeginVerticalStack();
 				// EditableText Example
-				if (SlateIM::EditableText(TypedText, TEXT("Type In Here")))
 				{
-					UE_LOG(LogSlateIM, Log, TEXT("Editable Text Changed: %s"), *TypedText);
-				}
-				if (SlateIM::IsFocused(SlateIM::EFocusDepth::IncludingDescendants))
-				{
-					SlateIM::BeginPopUp();
-					SlateIM::Text(TEXT("Pop Up Test:"));
-					SlateIM::Image("AppIcon");
-					SlateIM::EndPopUp();
+					SlateIM::BeginHorizontalStack();
+					SlateIM::EditableText(ComboItemToAdd, TEXT("Add Combo Item"));
+					if (SlateIM::IsFocused(SlateIM::EFocusDepth::IncludingDescendants))
+					{
+						SlateIM::BeginPopUp();
+						SlateIM::Text(TEXT("Enter the value of a new item to add to the combo box"));
+						SlateIM::EndPopUp();
+					}
+
+					bool bDisableAddButton = ComboItemToAdd.IsEmpty();
+					if (bDisableAddButton)
+					{
+						SlateIM::BeginDisabledState();
+					}
+
+					if (SlateIM::Button(TEXT("Add Combo Item")))
+					{
+						bRefreshComboItems = true;
+						ComboBoxItems.Add(MoveTemp(ComboItemToAdd));
+						ComboItemToAdd.Reset();
+					}
+
+					if (bDisableAddButton)
+					{
+						SlateIM::EndDisabledState();
+					}
+					SlateIM::EndHorizontalStack();
 				}
 				SlateIM::EndVerticalStack();
 			}
@@ -154,7 +172,7 @@ void FSlateIMTestWidget::Draw()
 			}
 
 			// ComboBox example
-			if (SlateIM::ComboBox(ComboBoxItems, SelectedItemIndex))
+			if (SlateIM::ComboBox(ComboBoxItems, SelectedItemIndex, bRefreshComboItems))
 			{
 				UE_LOG(LogSlateIM, Log, TEXT("Combo Box Item %s chosen"), *ComboBoxItems[SelectedItemIndex]);
 			}
@@ -189,7 +207,7 @@ void FSlateIMTestWidget::Draw()
 			SlateIM::EndVerticalStack();
 
 			// SelectionList example
-			SlateIM::SelectionList(ComboBoxItems, SelectedItem, false);
+			SlateIM::SelectionList(ComboBoxItems, SelectedItem, bRefreshComboItems);
 
 			SlateIM::EndHorizontalStack();
 
@@ -741,6 +759,8 @@ void FSlateIMTestWidget::Draw()
 	}
 	SlateIM::EndScrollBox();
 	SlateIM::EndBorder();
+
+	bRefreshComboItems = false;
 }
 
 void FSlateStyleBrowser::DrawWindow(float DeltaTime)
@@ -1008,7 +1028,8 @@ void FSlateStyleBrowser::DrawWindow(float DeltaTime)
 					SlateIM::HAlign(HAlign_Center);
 					SlateIM::VAlign(VAlign_Center);
 					SlateIM::Fill();
-					SlateIM::ComboBox(Options, SelectedComboIndex, WidgetStyle.Value);
+					constexpr bool bForceRefresh = false;
+					SlateIM::ComboBox(Options, SelectedComboIndex, bForceRefresh, WidgetStyle.Value);
 				}
 				
 				DrawCopyCell(WidgetStyle.Key);
diff --git a/Engine/Plugins/Experimental/SlateIM/Source/SlateIM/Private/Misc/SlateIMExamples.h b/Engine/Plugins/Experimental/SlateIM/Source/SlateIM/Private/Misc/SlateIMExamples.h
--- a/Engine/Plugins/Experimental/SlateIM/Source/SlateIM/Private/Misc/SlateIMExamples.h
+++ b/Engine/Plugins/Experimental/SlateIM/Source/SlateIM/Private/Misc/SlateIMExamples.h
@@ -19,7 +19,7 @@ private:
 	double CurrentTime = 0;
 	double TimeSinceLastUpdate = 0;
 	FString TimeText;
-	FString TypedText;
+	FString ComboItemToAdd;
 	bool CheckState = false;
 	ECheckBoxState CheckStateEnum = ECheckBoxState::Undetermined;
 	int32 SelectedItemIndex = 0;
@@ -36,7 +36,8 @@ private:
 	int32 LiveNumItems = NumItems;
 	FString NumItemsText = FString::FromInt(NumItems);
 	bool bShouldLiveUpdateTable = false;
-	const TArray<FString> ComboBoxItems = { TEXT("Option 1"), TEXT("Option 2"), TEXT("Option 3"), TEXT("Option 4") };
+	TArray<FString> ComboBoxItems = { TEXT("Option 1"), TEXT("Option 2"), TEXT("Option 3"), TEXT("Option 4") };
+	bool bRefreshComboItems = false;
 	TRingBuffer<double> SquareGraphValues;
 	TRingBuffer<FVector2D> SinGraphValues = { FVector2D(0, FMath::Sin(0.0)) };
 	TRingBuffer<FVector2D> CosGraphValues = { FVector2D(0, FMath::Cos(0.0)) };
diff --git a/Engine/Plugins/Experimental/SlateIM/Source/SlateIM/Public/SlateIM.h b/Engine/Plugins/Experimental/SlateIM/Source/SlateIM/Public/SlateIM.h
--- a/Engine/Plugins/Experimental/SlateIM/Source/SlateIM/Public/SlateIM.h
+++ b/Engine/Plugins/Experimental/SlateIM/Source/SlateIM/Public/SlateIM.h
@@ -491,25 +491,26 @@ namespace SlateIM
 	 * 
 	 * @param ComboItems The options the user can choose from
 	 * @param InOutSelectedItemIndex The index of the currently selected option
+	 * @param bForceRefresh (optional) Whether to force a refresh of the available options or the selected option (set to true for a frame when changing the list of options or manually setting the selected index)
 	 * @param InComboStyle (optional) A Combo Box Style to override the default style. Changes to the style after the first invocation are not reflected,
 	 *					   provide a new style object to update the visuals at runtime
 	 *					   
 	 * @return Whether the user changed the selected option this frame
 	 */
-	SLATEIM_API bool ComboBox(const TArray<FString>& ComboItems, int32& InOutSelectedItemIndex, const FComboBoxStyle* InComboStyle = nullptr);
+	SLATEIM_API bool ComboBox(const TArray<FString>& ComboItems, int32& InOutSelectedItemIndex, bool bForceRefresh = false, const FComboBoxStyle* InComboStyle = nullptr);
 
 	/**
 	 * Display a list of text options
 	 * 
 	 * @param ListItems The list of options the user can choose from
 	 * @param InOutSelectedItemIndex The index of the currently selected option
-	 * @param bForceRefresh Whether to force a refresh of the available options or the selected option (set to true for a frame when changing the list of options or manually setting the selected index)
+	 * @param bForceRefresh (optional) Whether to force a refresh of the available options or the selected option (set to true for a frame when changing the list of options or manually setting the selected index)
 	 * @param InStyle (optional) A Table View Style to override the default style. Changes to the style after the first invocation are not reflected,
 	 *				  provide a new style object to update the visuals at runtime
 	 *					   
 	 * @return Whether the user changed the selected option this frame
 	 */
-	SLATEIM_API bool SelectionList(const TArray<FString>& ListItems, int32& InOutSelectedItemIndex, bool bForceRefresh, const FTableViewStyle* InStyle = nullptr);
+	SLATEIM_API bool SelectionList(const TArray<FString>& ListItems, int32& InOutSelectedItemIndex, bool bForceRefresh = false, const FTableViewStyle* InStyle = nullptr);
 
 	/**
 	 * Create a block of empty space
