SlateIM: Add functions for querying inputs

SlateIM roots now wrap their contents with the SImWrapper widget which is responsible for tracking input state.

Allows querying for key presses/holds/releases and analog values.
Only works for input events that make it to the SlateIM root, so the SlateIM widget needs to be focused and input events that are handled by subwidgets (like editable text or buttons) cannot be queried for.

#rnx

diff --git a/Engine/Plugins/Experimental/SlateIM/Source/SlateIM/Private/Implementation/SlateIM_Inputs.cpp b/Engine/Plugins/Experimental/SlateIM/Source/SlateIM/Private/Implementation/SlateIM_Inputs.cpp
new file mode 100644
--- /dev/null
+++ b/Engine/Plugins/Experimental/SlateIM/Source/SlateIM/Private/Implementation/SlateIM_Inputs.cpp
@@ -0,0 +1,52 @@
+﻿// Copyright Epic Games, Inc. All Rights Reserved.
+
+#include "Misc/SlateIMInputState.h"
+#include "Misc/SlateIMManager.h"
+#include "SlateIM.h"
+
+namespace SlateIM
+{
+	bool IsKeyPressed(const FKey& InKey)
+	{
+		if (TSharedPtr<ISlateIMRoot> Root = FSlateIMManager::Get().GetCurrentRoot().RootWidget)
+		{
+			const ESlateIMKeyState* KeyStatePtr = Root->GetInputState().KeyStateMap.Find(InKey);
+			return KeyStatePtr != nullptr && *KeyStatePtr == ESlateIMKeyState::Pressed;
+		}
+
+		return false;
+	}
+
+	bool IsKeyHeld(const FKey& InKey)
+	{
+		if (TSharedPtr<ISlateIMRoot> Root = FSlateIMManager::Get().GetCurrentRoot().RootWidget)
+		{
+			const ESlateIMKeyState* KeyStatePtr = Root->GetInputState().KeyStateMap.Find(InKey);
+			return KeyStatePtr != nullptr && (*KeyStatePtr == ESlateIMKeyState::Pressed || *KeyStatePtr == ESlateIMKeyState::Held);
+		}
+
+		return false;
+	}
+
+	bool IsKeyReleased(const FKey& InKey)
+	{
+		if (TSharedPtr<ISlateIMRoot> Root = FSlateIMManager::Get().GetCurrentRoot().RootWidget)
+		{
+			const ESlateIMKeyState* KeyStatePtr = Root->GetInputState().KeyStateMap.Find(InKey);
+			return KeyStatePtr != nullptr && *KeyStatePtr == ESlateIMKeyState::Released;
+		}
+
+		return false;
+	}
+
+	float GetKeyAnalogValue(const FKey& InKey)
+	{
+		if (TSharedPtr<ISlateIMRoot> Root = FSlateIMManager::Get().GetCurrentRoot().RootWidget)
+		{
+			const float* AnalogValuePtr = Root->GetInputState().AnalogValueMap.Find(InKey);
+			return (AnalogValuePtr != nullptr) ? *AnalogValuePtr : 0.f;
+		}
+
+		return 0.f;
+	}
+}
diff --git a/Engine/Plugins/Experimental/SlateIM/Source/SlateIM/Private/Misc/SlateIMExamples.cpp b/Engine/Plugins/Experimental/SlateIM/Source/SlateIM/Private/Misc/SlateIMExamples.cpp
--- a/Engine/Plugins/Experimental/SlateIM/Source/SlateIM/Private/Misc/SlateIMExamples.cpp
+++ b/Engine/Plugins/Experimental/SlateIM/Source/SlateIM/Private/Misc/SlateIMExamples.cpp
@@ -65,6 +65,24 @@ namespace SlateIM::Private
 			return Styles;
 		}
 	};
+
+	FLinearColor GetKeyStateColor(const FKey& Key)
+	{
+		if (SlateIM::IsKeyPressed(Key))
+		{
+			return FLinearColor::Green;
+		}
+		else if (SlateIM::IsKeyHeld(Key))
+		{
+			return FLinearColor::Blue;	
+		}
+		else if (SlateIM::IsKeyReleased(Key))
+		{
+			return FLinearColor::Red;	
+		}
+
+		return FLinearColor(0.1, 0.1, 0.1);
+	}
 }
 
 void FSlateIMTestWidget::Draw()
@@ -752,6 +770,87 @@ void FSlateIMTestWidget::Draw()
 			SlateIM::EndTable();
 		}
 
+		// Inputs example
+		{
+			SlateIM::BeginHorizontalStack();
+			{
+				SlateIM::BeginVerticalStack();
+				SlateIM::BeginHorizontalStack();
+				SlateIM::Spacer(FVector2D(50, 50));
+				{
+					WBrush = FSlateColorBrush(SlateIM::Private::GetKeyStateColor(EKeys::W));
+					SlateIM::MinWidth(50.f);
+					SlateIM::MinHeight(50.f);
+					SlateIM::MaxWidth(50.f);
+					SlateIM::MaxHeight(50.f);
+					SlateIM::BeginBorder(&WBrush);
+					SlateIM::HAlign(HAlign_Center);
+					SlateIM::VAlign(VAlign_Center);
+					SlateIM::Text(TEXT("W"), FLinearColor::White);
+					SlateIM::EndBorder();
+				}
+				SlateIM::Spacer(FVector2D(50, 50));
+				SlateIM::EndHorizontalStack();
+				SlateIM::BeginHorizontalStack();
+				{
+					ABrush = FSlateColorBrush(SlateIM::Private::GetKeyStateColor(EKeys::A));
+					SlateIM::MinWidth(50.f);
+					SlateIM::MinHeight(50.f);
+					SlateIM::MaxWidth(50.f);
+					SlateIM::MaxHeight(50.f);
+					SlateIM::BeginBorder(&ABrush);
+					SlateIM::HAlign(HAlign_Center);
+					SlateIM::VAlign(VAlign_Center);
+					SlateIM::Text(TEXT("A"), FLinearColor::White);
+					SlateIM::EndBorder();
+				}
+				{
+					SBrush = FSlateColorBrush(SlateIM::Private::GetKeyStateColor(EKeys::S));
+					SlateIM::MinWidth(50.f);
+					SlateIM::MinHeight(50.f);
+					SlateIM::MaxWidth(50.f);
+					SlateIM::MaxHeight(50.f);
+					SlateIM::BeginBorder(&SBrush);
+					SlateIM::HAlign(HAlign_Center);
+					SlateIM::VAlign(VAlign_Center);
+					SlateIM::Text(TEXT("S"), FLinearColor::White);
+					SlateIM::EndBorder();
+				}
+				{
+					DBrush = FSlateColorBrush(SlateIM::Private::GetKeyStateColor(EKeys::D));
+					SlateIM::MinWidth(50.f);
+					SlateIM::MinHeight(50.f);
+					SlateIM::MaxWidth(50.f);
+					SlateIM::MaxHeight(50.f);
+					SlateIM::BeginBorder(&DBrush);
+					SlateIM::HAlign(HAlign_Center);
+					SlateIM::VAlign(VAlign_Center);
+					SlateIM::Text(TEXT("D"), FLinearColor::White);
+					SlateIM::EndBorder();
+				}
+				SlateIM::EndHorizontalStack();
+				SlateIM::EndVerticalStack();
+			}
+
+			{
+				SlateIM::BeginVerticalStack();
+				SlateIM::Text(TEXT("Right Stick X-value"));
+				const float NormalizedAnalogXValue = (1.f + SlateIM::GetKeyAnalogValue(EKeys::Gamepad_RightX)) * 0.5f;
+				SlateIM::VAlign(VAlign_Center);
+				SlateIM::ProgressBar(NormalizedAnalogXValue);
+				
+				SlateIM::Fill();
+				SlateIM::Spacer(FVector2D(1.f, 1.f));
+				
+				SlateIM::Text(TEXT("Right Stick Y-value"));
+				const float NormalizedAnalogYValue = (1.f + SlateIM::GetKeyAnalogValue(EKeys::Gamepad_RightY)) * 0.5f;
+				SlateIM::VAlign(VAlign_Center);
+				SlateIM::ProgressBar(NormalizedAnalogYValue);
+				SlateIM::EndVerticalStack();
+			}
+			SlateIM::EndHorizontalStack();
+		}
+
 		if (bShouldBeDisabled)
 		{
 			SlateIM::EndDisabledState();
diff --git a/Engine/Plugins/Experimental/SlateIM/Source/SlateIM/Private/Misc/SlateIMExamples.h b/Engine/Plugins/Experimental/SlateIM/Source/SlateIM/Private/Misc/SlateIMExamples.h
--- a/Engine/Plugins/Experimental/SlateIM/Source/SlateIM/Private/Misc/SlateIMExamples.h
+++ b/Engine/Plugins/Experimental/SlateIM/Source/SlateIM/Private/Misc/SlateIMExamples.h
@@ -43,6 +43,11 @@ private:
 	TRingBuffer<FVector2D> CosGraphValues = { FVector2D(0, FMath::Cos(0.0)) };
 	TRingBuffer<FVector2D> TanGraphValues = { FVector2D(0, FMath::Tan(0.0)) };
 
+	FSlateBrush WBrush;
+	FSlateBrush ABrush;
+	FSlateBrush SBrush;
+	FSlateBrush DBrush;
+
 #if WITH_ENGINE
 	TSoftObjectPtr<UTexture2D> RedIcon = TSoftObjectPtr<UTexture2D>(FSoftObjectPath(TEXT("/Engine/EngineResources/AICON-Red.AICON-Red")));
 	TSoftObjectPtr<UTexture2D> GreenIcon = TSoftObjectPtr<UTexture2D>(FSoftObjectPath(TEXT("/Engine/EngineResources/AICON-Green.AICON-Green")));
diff --git a/Engine/Plugins/Experimental/SlateIM/Source/SlateIM/Private/Misc/SlateIMInputState.h b/Engine/Plugins/Experimental/SlateIM/Source/SlateIM/Private/Misc/SlateIMInputState.h
new file mode 100644
--- /dev/null
+++ b/Engine/Plugins/Experimental/SlateIM/Source/SlateIM/Private/Misc/SlateIMInputState.h
@@ -0,0 +1,19 @@
+﻿// Copyright Epic Games, Inc. All Rights Reserved.
+
+#pragma once
+
+#include "InputCoreTypes.h"
+
+enum class ESlateIMKeyState : uint8
+{
+	Idle,
+	Pressed,
+	Held,
+	Released,
+};
+
+struct FSlateIMInputState
+{
+	TMap<FKey, ESlateIMKeyState> KeyStateMap;
+	TMap<FKey, float> AnalogValueMap;
+};
diff --git a/Engine/Plugins/Experimental/SlateIM/Source/SlateIM/Private/Roots/ISlateIMRoot.h b/Engine/Plugins/Experimental/SlateIM/Source/SlateIM/Private/Roots/ISlateIMRoot.h
--- a/Engine/Plugins/Experimental/SlateIM/Source/SlateIM/Private/Roots/ISlateIMRoot.h
+++ b/Engine/Plugins/Experimental/SlateIM/Source/SlateIM/Private/Roots/ISlateIMRoot.h
@@ -7,6 +7,7 @@
 #include "Misc/SlateIMTypeChecking.h"
 #include "Templates/SharedPointer.h"
 
+struct FSlateIMInputState;
 struct FSlateIMSlotData;
 class SWidget;
 
@@ -15,4 +16,6 @@ class ISlateIMRoot : public ISlateIMTypeChecking
 public:
 	virtual void UpdateChild(TSharedRef<SWidget> Child, const FSlateIMSlotData& AlignmentData) = 0;
 	virtual bool IsVisible() const = 0;
+
+	virtual FSlateIMInputState& GetInputState() = 0;
 };
diff --git a/Engine/Plugins/Experimental/SlateIM/Source/SlateIM/Private/Roots/SlateIMExposedRoot.cpp b/Engine/Plugins/Experimental/SlateIM/Source/SlateIM/Private/Roots/SlateIMExposedRoot.cpp
--- a/Engine/Plugins/Experimental/SlateIM/Source/SlateIM/Private/Roots/SlateIMExposedRoot.cpp
+++ b/Engine/Plugins/Experimental/SlateIM/Source/SlateIM/Private/Roots/SlateIMExposedRoot.cpp
@@ -3,11 +3,11 @@
 #include "SlateIMExposedRoot.h"
 
 #include "Misc/SlateIMSlotData.h"
-#include "Widgets/Layout/SBox.h"
+#include "Widgets/SImWrapper.h"
 
 
 FSlateIMExposedRoot::FSlateIMExposedRoot()
-	: ExposedWidget(SNew(SBox))
+	: ExposedWidget(SNew(SImWrapper))
 {
 }
 
@@ -24,6 +24,11 @@ void FSlateIMExposedRoot::UpdateChild(TSharedRef<SWidget> Child, const FSlateIMS
 	ExposedWidget->SetVAlign(AlignmentData.VerticalAlignment);
 }
 
+FSlateIMInputState& FSlateIMExposedRoot::GetInputState()
+{
+	return ExposedWidget->InputState;
+}
+
 TSharedRef<SWidget> FSlateIMExposedRoot::GetExposedWidget() const
 {
 	return ExposedWidget;
diff --git a/Engine/Plugins/Experimental/SlateIM/Source/SlateIM/Private/Roots/SlateIMExposedRoot.h b/Engine/Plugins/Experimental/SlateIM/Source/SlateIM/Private/Roots/SlateIMExposedRoot.h
--- a/Engine/Plugins/Experimental/SlateIM/Source/SlateIM/Private/Roots/SlateIMExposedRoot.h
+++ b/Engine/Plugins/Experimental/SlateIM/Source/SlateIM/Private/Roots/SlateIMExposedRoot.h
@@ -4,7 +4,7 @@
 
 #include "ISlateIMRoot.h"
 
-class SBox;
+class SImWrapper;
 
 class FSlateIMExposedRoot : public ISlateIMRoot
 {
@@ -16,10 +16,11 @@ public:
 
 	virtual void UpdateChild(TSharedRef<SWidget> Child, const FSlateIMSlotData& AlignmentData) override;
 	virtual bool IsVisible() const override { return true; }
+	virtual FSlateIMInputState& GetInputState() override;
 	
 	TSharedRef<SWidget> GetExposedWidget() const;
 	
 private:
-	TSharedRef<SBox> ExposedWidget;
+	TSharedRef<SImWrapper> ExposedWidget;
 
 };
diff --git a/Engine/Plugins/Experimental/SlateIM/Source/SlateIM/Private/Roots/SlateIMViewportRoot.cpp b/Engine/Plugins/Experimental/SlateIM/Source/SlateIM/Private/Roots/SlateIMViewportRoot.cpp
--- a/Engine/Plugins/Experimental/SlateIM/Source/SlateIM/Private/Roots/SlateIMViewportRoot.cpp
+++ b/Engine/Plugins/Experimental/SlateIM/Source/SlateIM/Private/Roots/SlateIMViewportRoot.cpp
@@ -7,13 +7,12 @@
 #if WITH_EDITOR
 #include "IAssetViewport.h"
 #endif
+
 #include "SlateIM.h"
 #include "Engine/GameViewportClient.h"
 #include "Engine/LocalPlayer.h"
-#endif
+#include "Widgets/SImWrapper.h"
 
-
-#if WITH_ENGINE
 FSlateIMViewportRoot::FSlateIMViewportRoot(UGameViewportClient* InGameViewport)
 	: GameViewport(InGameViewport)
 {
@@ -38,12 +37,12 @@ FSlateIMViewportRoot::FSlateIMViewportRoot(TSharedPtr<IAssetViewport> InAssetVie
 
 FSlateIMViewportRoot::~FSlateIMViewportRoot()
 {
-	if (ViewportCanvas)
+	if (ViewportRoot)
 	{
 #if WITH_EDITOR
 		if (TSharedPtr<IAssetViewport> AssetViewportPinned = AssetViewport.Pin())
 		{
-			AssetViewportPinned->RemoveOverlayWidget(ViewportCanvas.ToSharedRef());
+			AssetViewportPinned->RemoveOverlayWidget(ViewportRoot.ToSharedRef());
 		}
 		else
 #endif // WITH_EDITOR
@@ -51,11 +50,11 @@ FSlateIMViewportRoot::~FSlateIMViewportRoot()
 		{
 			if (LocalPlayer.IsValid())
 			{
-				GameViewport->RemoveViewportWidgetForPlayer(LocalPlayer.Get(), ViewportCanvas.ToSharedRef());
+				GameViewport->RemoveViewportWidgetForPlayer(LocalPlayer.Get(), ViewportRoot.ToSharedRef());
 			}
 			else
 			{
-				GameViewport->RemoveGameLayerWidget(ViewportCanvas.ToSharedRef());
+				GameViewport->RemoveGameLayerWidget(ViewportRoot.ToSharedRef());
 			}
 		}
 	}
@@ -69,21 +68,29 @@ void FSlateIMViewportRoot::UpdateChild(TSharedRef<SWidget> Child, const FSlateIM
 	}
 }
 
+FSlateIMInputState& FSlateIMViewportRoot::GetInputState()
+{
+	check(ViewportRoot.IsValid());
+	return ViewportRoot->InputState;
+}
+
 void FSlateIMViewportRoot::UpdateViewport(const SlateIM::FViewportRootLayout& Layout)
 {
-	if (!ViewportCanvas)
+	if (!ViewportRoot)
 	{
-		ViewportCanvas =
+		ViewportRoot = SNew(SImWrapper)
+		[
 			SNew(SConstraintCanvas)
 			+ SConstraintCanvas::Slot()
 			.Anchors(Layout.Anchors)
 			.Alignment(Layout.Alignment)
-			.Expose(Slot);
+			.Expose(Slot)
+		];
 		
 #if WITH_EDITOR
 		if (TSharedPtr<IAssetViewport> AssetViewportPinned = AssetViewport.Pin())
 		{
-			AssetViewportPinned->AddOverlayWidget(ViewportCanvas.ToSharedRef());
+			AssetViewportPinned->AddOverlayWidget(ViewportRoot.ToSharedRef());
 		}
 		else
 #endif // WITH_EDITOR
@@ -92,11 +99,11 @@ void FSlateIMViewportRoot::UpdateViewport(const SlateIM::FViewportRootLayout& La
 			int ZOrder = 10;
 			if (LocalPlayer.IsValid())
 			{
-				GameViewport->AddViewportWidgetForPlayer(LocalPlayer.Get(), ViewportCanvas.ToSharedRef(), ZOrder);
+				GameViewport->AddViewportWidgetForPlayer(LocalPlayer.Get(), ViewportRoot.ToSharedRef(), ZOrder);
 			}
 			else
 			{
-				GameViewport->AddGameLayerWidget(ViewportCanvas.ToSharedRef(), ZOrder);
+				GameViewport->AddGameLayerWidget(ViewportRoot.ToSharedRef(), ZOrder);
 			}
 		}
 	}
diff --git a/Engine/Plugins/Experimental/SlateIM/Source/SlateIM/Private/Roots/SlateIMViewportRoot.h b/Engine/Plugins/Experimental/SlateIM/Source/SlateIM/Private/Roots/SlateIMViewportRoot.h
--- a/Engine/Plugins/Experimental/SlateIM/Source/SlateIM/Private/Roots/SlateIMViewportRoot.h
+++ b/Engine/Plugins/Experimental/SlateIM/Source/SlateIM/Private/Roots/SlateIMViewportRoot.h
@@ -7,6 +7,8 @@
 #include "UObject/WeakObjectPtrTemplates.h"
 #include "Widgets/Layout/SConstraintCanvas.h"
 
+class SImWrapper;
+
 namespace SlateIM
 {
 	struct FViewportRootLayout;
@@ -36,6 +38,7 @@ public:
 	{
 		return true;
 	}
+	virtual FSlateIMInputState& GetInputState() override;
 
 	void UpdateViewport(const SlateIM::FViewportRootLayout& Layout);
 
@@ -46,7 +49,7 @@ public:
 #endif
 
 private:
-	TSharedPtr<SConstraintCanvas> ViewportCanvas;
+	TSharedPtr<SImWrapper> ViewportRoot;
 	SConstraintCanvas::FSlot* Slot = nullptr;
 };
 #endif // WITH_ENGINE
diff --git a/Engine/Plugins/Experimental/SlateIM/Source/SlateIM/Private/Roots/SlateIMWindowRoot.cpp b/Engine/Plugins/Experimental/SlateIM/Source/SlateIM/Private/Roots/SlateIMWindowRoot.cpp
--- a/Engine/Plugins/Experimental/SlateIM/Source/SlateIM/Private/Roots/SlateIMWindowRoot.cpp
+++ b/Engine/Plugins/Experimental/SlateIM/Source/SlateIM/Private/Roots/SlateIMWindowRoot.cpp
@@ -4,6 +4,7 @@
 #include "SlateIMWindowRoot.h"
 
 #include "Framework/Application/SlateApplication.h"
+#include "Widgets/SImWrapper.h"
 
 FSlateIMWindowRoot::FSlateIMWindowRoot(TSharedRef<SWindow> Window)
 	: RootWindow(Window)
@@ -16,13 +17,21 @@ FSlateIMWindowRoot::~FSlateIMWindowRoot()
 		FSlateApplication::Get().RequestDestroyWindow(RootWindow.Pin().ToSharedRef());
 		RootWindow = nullptr;
 	}
+	
+	WindowRootWidget.Reset();
 }
 
 void FSlateIMWindowRoot::UpdateChild(TSharedRef<SWidget> Child, const FSlateIMSlotData& AlignmentData)
 {
 	check(RootWindow != Child);
 
-	RootWindow.Pin()->SetContent(Child);
+	if (!WindowRootWidget.IsValid())
+	{
+		WindowRootWidget = SNew(SImWrapper);
+		RootWindow.Pin()->SetContent(WindowRootWidget.ToSharedRef());
+	}
+
+	WindowRootWidget->SetContent(Child);
 }
 
 bool FSlateIMWindowRoot::IsVisible() const
@@ -30,6 +39,12 @@ bool FSlateIMWindowRoot::IsVisible() const
 	return RootWindow.IsValid();
 }
 
+FSlateIMInputState& FSlateIMWindowRoot::GetInputState()
+{
+	check(WindowRootWidget.IsValid());
+	return WindowRootWidget->InputState;
+}
+
 void FSlateIMWindowRoot::UpdateWindow(const FStringView& Title)
 {
 	if (TSharedPtr<SWindow> RootWindowPin = RootWindow.Pin())
diff --git a/Engine/Plugins/Experimental/SlateIM/Source/SlateIM/Private/Roots/SlateIMWindowRoot.h b/Engine/Plugins/Experimental/SlateIM/Source/SlateIM/Private/Roots/SlateIMWindowRoot.h
--- a/Engine/Plugins/Experimental/SlateIM/Source/SlateIM/Private/Roots/SlateIMWindowRoot.h
+++ b/Engine/Plugins/Experimental/SlateIM/Source/SlateIM/Private/Roots/SlateIMWindowRoot.h
@@ -4,6 +4,7 @@
 
 #include "ISlateIMRoot.h"
 
+class SImWrapper;
 class SWindow;
 
 class FSlateIMWindowRoot : public ISlateIMRoot
@@ -16,10 +17,12 @@ public:
 
 	virtual void UpdateChild(TSharedRef<SWidget> Child, const FSlateIMSlotData& AlignmentData) override;
 	virtual bool IsVisible() const override;
+	virtual FSlateIMInputState& GetInputState() override;
 
 	void UpdateWindow(const FStringView& Title);
 	
 private:
 	TWeakPtr<SWindow> RootWindow;
+	TSharedPtr<SImWrapper> WindowRootWidget;
 
 };
diff --git a/Engine/Plugins/Experimental/SlateIM/Source/SlateIM/Private/Widgets/SImWrapper.cpp b/Engine/Plugins/Experimental/SlateIM/Source/SlateIM/Private/Widgets/SImWrapper.cpp
new file mode 100644
--- /dev/null
+++ b/Engine/Plugins/Experimental/SlateIM/Source/SlateIM/Private/Widgets/SImWrapper.cpp
@@ -0,0 +1,116 @@
+﻿// Copyright Epic Games, Inc. All Rights Reserved.
+
+
+#include "SImWrapper.h"
+
+
+SImWrapper::~SImWrapper()
+{
+}
+
+void SImWrapper::Construct(const FArguments& InArgs)
+{
+	ChildSlot
+	.HAlign(HAlign_Fill)
+	.VAlign(VAlign_Fill)
+	[
+		InArgs._Content.Widget
+	];
+}
+
+void SImWrapper::Tick(const FGeometry& AllottedGeometry, const double InCurrentTime, const float InDeltaTime)
+{
+	SCompoundWidget::Tick(AllottedGeometry, InCurrentTime, InDeltaTime);
+	
+	// Inputs are only considered to be released for 1 frame
+	for (auto It = InputState.KeyStateMap.CreateIterator(); It; ++It)
+	{
+		if (It.Value() == ESlateIMKeyState::Released)
+		{
+			It.Value() = ESlateIMKeyState::Idle;
+		}
+	}
+
+	// Reset mouse wheel axis since we don't get "0" events for it
+	if (float* MouseWheelValuePtr = InputState.AnalogValueMap.Find(EKeys::MouseWheelAxis))
+	{
+		*MouseWheelValuePtr = 0;
+	}
+}
+
+bool SImWrapper::SupportsKeyboardFocus() const
+{
+	return true;
+}
+
+FReply SImWrapper::OnKeyDown(const FGeometry& MyGeometry, const FKeyEvent& KeyEvent)
+{
+	InputState.KeyStateMap.FindOrAdd(KeyEvent.GetKey()) = (KeyEvent.IsRepeat() ? ESlateIMKeyState::Held : ESlateIMKeyState::Pressed);
+	return SCompoundWidget::OnKeyDown(MyGeometry, KeyEvent);
+}
+
+FReply SImWrapper::OnKeyUp(const FGeometry& MyGeometry, const FKeyEvent& KeyEvent)
+{
+	InputState.KeyStateMap.FindOrAdd(KeyEvent.GetKey()) = ESlateIMKeyState::Released;
+	return SCompoundWidget::OnKeyUp(MyGeometry, KeyEvent);
+}
+
+FReply SImWrapper::OnAnalogValueChanged(const FGeometry& MyGeometry, const FAnalogInputEvent& AnalogInputEvent)
+{
+	InputState.AnalogValueMap.FindOrAdd(AnalogInputEvent.GetKey()) = AnalogInputEvent.GetAnalogValue();
+	return SCompoundWidget::OnAnalogValueChanged(MyGeometry, AnalogInputEvent);
+}
+
+FReply SImWrapper::OnMouseButtonDown(const FGeometry& MyGeometry, const FPointerEvent& MouseEvent)
+{
+	InputState.KeyStateMap.FindOrAdd(MouseEvent.GetEffectingButton()) = (MouseEvent.IsRepeat() ? ESlateIMKeyState::Held : ESlateIMKeyState::Pressed);
+	return SCompoundWidget::OnMouseButtonDown(MyGeometry, MouseEvent);
+}
+
+FReply SImWrapper::OnMouseButtonDoubleClick(const FGeometry& MyGeometry, const FPointerEvent& MouseEvent)
+{
+	InputState.KeyStateMap.FindOrAdd(MouseEvent.GetEffectingButton()) = (MouseEvent.IsRepeat() ? ESlateIMKeyState::Held : ESlateIMKeyState::Pressed);
+	return SCompoundWidget::OnMouseButtonDoubleClick(MyGeometry, MouseEvent);
+}
+
+FReply SImWrapper::OnMouseButtonUp(const FGeometry& MyGeometry, const FPointerEvent& MouseEvent)
+{
+	InputState.KeyStateMap.FindOrAdd(MouseEvent.GetEffectingButton()) = ESlateIMKeyState::Released;
+	return SCompoundWidget::OnMouseButtonUp(MyGeometry, MouseEvent);
+}
+
+FReply SImWrapper::OnMouseWheel(const FGeometry& MyGeometry, const FPointerEvent& MouseEvent)
+{
+	InputState.AnalogValueMap.FindOrAdd(EKeys::MouseWheelAxis) = MouseEvent.GetWheelDelta();
+	return SCompoundWidget::OnMouseWheel(MyGeometry, MouseEvent);
+}
+
+FReply SImWrapper::OnMouseMove(const FGeometry& MyGeometry, const FPointerEvent& MouseEvent)
+{
+	InputState.AnalogValueMap.FindOrAdd(EKeys::MouseX) = MouseEvent.GetCursorDelta().X;
+	InputState.AnalogValueMap.FindOrAdd(EKeys::MouseY) = MouseEvent.GetCursorDelta().Y;
+	return SCompoundWidget::OnMouseMove(MyGeometry, MouseEvent);
+}
+
+void SImWrapper::SetContent(const TSharedRef<SWidget>& InContent)
+{
+	ChildSlot
+	[
+		InContent
+	];
+}
+
+void SImWrapper::SetHAlign(EHorizontalAlignment HAlign)
+{
+	ChildSlot.SetHorizontalAlignment(HAlign);
+}
+
+void SImWrapper::SetVAlign(EVerticalAlignment VAlign)
+{
+	ChildSlot.SetVerticalAlignment(VAlign);
+}
+
+void SImWrapper::SetPadding(TAttribute<FMargin> InPadding)
+{
+	ChildSlot.SetPadding(InPadding);
+}
diff --git a/Engine/Plugins/Experimental/SlateIM/Source/SlateIM/Private/Widgets/SImWrapper.h b/Engine/Plugins/Experimental/SlateIM/Source/SlateIM/Private/Widgets/SImWrapper.h
new file mode 100644
--- /dev/null
+++ b/Engine/Plugins/Experimental/SlateIM/Source/SlateIM/Private/Widgets/SImWrapper.h
@@ -0,0 +1,39 @@
+﻿// Copyright Epic Games, Inc. All Rights Reserved.
+
+#pragma once
+
+#include "Misc/SlateIMInputState.h"
+#include "Widgets/SCompoundWidget.h"
+
+class SLATEIM_API SImWrapper : public SCompoundWidget
+{
+public:
+	SLATE_BEGIN_ARGS(SImWrapper)
+		{}
+		SLATE_DEFAULT_SLOT(FArguments, Content)
+	SLATE_END_ARGS()
+
+	virtual ~SImWrapper() override;
+
+	void Construct(const FArguments& InArgs);
+
+	virtual void Tick(const FGeometry& AllottedGeometry, const double InCurrentTime, const float InDeltaTime) override;
+
+	virtual bool SupportsKeyboardFocus() const override;
+
+	virtual FReply OnKeyDown(const FGeometry& MyGeometry, const FKeyEvent& KeyEvent) override;
+	virtual FReply OnKeyUp(const FGeometry& MyGeometry, const FKeyEvent& KeyEvent) override;
+	virtual FReply OnAnalogValueChanged(const FGeometry& MyGeometry, const FAnalogInputEvent& AnalogInputEvent) override;
+	virtual FReply OnMouseButtonDown(const FGeometry& MyGeometry, const FPointerEvent& MouseEvent) override;
+	virtual FReply OnMouseButtonDoubleClick(const FGeometry& MyGeometry, const FPointerEvent& MouseEvent) override;
+	virtual FReply OnMouseButtonUp(const FGeometry& MyGeometry, const FPointerEvent& MouseEvent) override;
+	virtual FReply OnMouseWheel(const FGeometry& MyGeometry, const FPointerEvent& MouseEvent) override;
+	virtual FReply OnMouseMove(const FGeometry& MyGeometry, const FPointerEvent& MouseEvent) override;
+
+	void SetContent(const TSharedRef< SWidget >& InContent);
+	void SetHAlign(EHorizontalAlignment HAlign);
+	void SetVAlign(EVerticalAlignment VAlign);
+	void SetPadding(TAttribute<FMargin> InPadding);
+
+	FSlateIMInputState InputState;
+};
diff --git a/Engine/Plugins/Experimental/SlateIM/Source/SlateIM/Public/SlateIM.h b/Engine/Plugins/Experimental/SlateIM/Source/SlateIM/Public/SlateIM.h
--- a/Engine/Plugins/Experimental/SlateIM/Source/SlateIM/Public/SlateIM.h
+++ b/Engine/Plugins/Experimental/SlateIM/Source/SlateIM/Public/SlateIM.h
@@ -24,6 +24,7 @@ struct FButtonStyle;
 struct FCheckBoxStyle;
 struct FComboBoxStyle;
 struct FEditableTextBoxStyle;
+struct FKey;
 struct FProgressBarStyle;
 struct FSlateBrush;
 struct FSliderStyle;
@@ -685,5 +686,47 @@ namespace SlateIM
 	 */
 	SLATEIM_API void GraphLine(const TArrayView<double>& Values, const FLinearColor& LineColor, float LineThickness, const FDoubleRange& ViewRange);
 #pragma endregion Graphing
+
+#pragma region Inputs
+	/**
+	 * Query whether a key was pressed this frame.
+	 *
+	 * Only considers input events that make it to the SlateIM root. Input events handled by child widgets will not be considered.
+	 * 
+	 * @param InKey The key to query the state of
+	 * @return true if the key was just pressed this frame
+	 */
+	SLATEIM_API bool IsKeyPressed(const FKey& InKey);
+	
+	/**
+	 * Query whether a key is being held this frame.
+	 *
+	 * Only considers input events that make it to the SlateIM root. Input events handled by child widgets will not be considered.
+	 * 
+	 * @param InKey The key to query the state of
+	 * @return true if the key was pressed or held this frame
+	 */
+	SLATEIM_API bool IsKeyHeld(const FKey& InKey);
+	
+	/**
+	 * Query whether a key was released this frame.
+	 *
+	 * Only considers input events that make it to the SlateIM root. Input events handled by child widgets will not be considered.
+	 * 
+	 * @param InKey The key to query the state of
+	 * @return true if the key was just released this frame
+	 */
+	SLATEIM_API bool IsKeyReleased(const FKey& InKey);
+	
+	/**
+	 * Retrieve the analog value for a key
+	 *
+	 * Only considers input events that make it to the SlateIM root. Input events handled by child widgets will not be considered.
+	 * 
+	 * @param InKey The key to query the state of
+	 * @return the last analog value the SlateIM root received for the specified key
+	 */
+	SLATEIM_API float GetKeyAnalogValue(const FKey& InKey);
+#pragma endregion Inputs
 }
 
