SlateIM: Expose ZOrder of SlateIM viewport roots

#rnx

diff --git a/Engine/Plugins/Experimental/SlateIM/Source/SlateIM/Private/Roots/SlateIMViewportRoot.cpp b/Engine/Plugins/Experimental/SlateIM/Source/SlateIM/Private/Roots/SlateIMViewportRoot.cpp
--- a/Engine/Plugins/Experimental/SlateIM/Source/SlateIM/Private/Roots/SlateIMViewportRoot.cpp
+++ b/Engine/Plugins/Experimental/SlateIM/Source/SlateIM/Private/Roots/SlateIMViewportRoot.cpp
@@ -90,20 +90,19 @@ void FSlateIMViewportRoot::UpdateViewport(const SlateIM::FViewportRootLayout& La
 #if WITH_EDITOR
 		if (TSharedPtr<IAssetViewport> AssetViewportPinned = AssetViewport.Pin())
 		{
-			AssetViewportPinned->AddOverlayWidget(ViewportRoot.ToSharedRef());
+			AssetViewportPinned->AddOverlayWidget(ViewportRoot.ToSharedRef(), Layout.ZOrder);
 		}
 		else
 #endif // WITH_EDITOR
 		if (GameViewport.IsValid())
 		{
-			int ZOrder = 10;
 			if (LocalPlayer.IsValid())
 			{
-				GameViewport->AddViewportWidgetForPlayer(LocalPlayer.Get(), ViewportRoot.ToSharedRef(), ZOrder);
+				GameViewport->AddViewportWidgetForPlayer(LocalPlayer.Get(), ViewportRoot.ToSharedRef(), Layout.ZOrder);
 			}
 			else
 			{
-				GameViewport->AddGameLayerWidget(ViewportRoot.ToSharedRef(), ZOrder);
+				GameViewport->AddGameLayerWidget(ViewportRoot.ToSharedRef(), Layout.ZOrder);
 			}
 		}
 	}
diff --git a/Engine/Plugins/Experimental/SlateIM/Source/SlateIM/Public/SlateIM.h b/Engine/Plugins/Experimental/SlateIM/Source/SlateIM/Public/SlateIM.h
--- a/Engine/Plugins/Experimental/SlateIM/Source/SlateIM/Public/SlateIM.h
+++ b/Engine/Plugins/Experimental/SlateIM/Source/SlateIM/Public/SlateIM.h
@@ -64,6 +64,9 @@ namespace SlateIM
 
 		/** Optional set the size (in slate units) of the root in the viewport. When unset, the root will auto-size for its content. */
 		TOptional<FVector2f> Size;
+
+		/** The ZOrder of the SlateIM widget in the viewport */
+		int32 ZOrder = 10000;
 	};
 #endif // WITH_EDITOR || WITH_ENGINE
 
