SlateIM: Fix issue where widgets placed after a table in the hierarchy could not receive input

This was caused by the internal widget tracking using the index of the widget within the full hierarchy, therefore assuming that the number of widgets would never change after initial construction. This change updates the tracking to refresh the widget's cached index every frame to match the latest hierarchy.

diff --git a/Engine/Plugins/Experimental/SlateIM/Source/SlateIM/Private/Implementation/SlateIM_Containers.cpp b/Engine/Plugins/Experimental/SlateIM/Source/SlateIM/Private/Implementation/SlateIM_Containers.cpp
--- a/Engine/Plugins/Experimental/SlateIM/Source/SlateIM/Private/Implementation/SlateIM_Containers.cpp
+++ b/Engine/Plugins/Experimental/SlateIM/Source/SlateIM/Private/Implementation/SlateIM_Containers.cpp
@@ -160,16 +160,16 @@ namespace SlateIM
 				ContainerWidget =
 					SNew(SImScrollBox)
 					.Orientation(Orientation)
-					.OnUserScrolled_Lambda([ActivationData = FSlateIMManager::Get().GetWidgetActivationData()](float Offset)
+					.OnUserScrolled_Lambda([ActivationData = Scope.GetOrCreateActivationMetadata().ToWeakPtr()](float Offset)
 					{
-						FSlateIMManager::Get().ActivateWidget(ActivationData);
+						FSlateIMManager::Get().ActivateWidget(ActivationData.Pin());
 					});
 
 				Scope.UpdateWidget(ContainerWidget);
 			}
 			else
 			{
-				bUserScrolled = FSlateIMManager::Get().IsWidgetActivatedThisFrame();
+				bUserScrolled = Scope.IsActivatedThisFrame();
 			}
 		}
 
diff --git a/Engine/Plugins/Experimental/SlateIM/Source/SlateIM/Private/Implementation/SlateIM_Widgets.cpp b/Engine/Plugins/Experimental/SlateIM/Source/SlateIM/Private/Implementation/SlateIM_Widgets.cpp
--- a/Engine/Plugins/Experimental/SlateIM/Source/SlateIM/Private/Implementation/SlateIM_Widgets.cpp
+++ b/Engine/Plugins/Experimental/SlateIM/Source/SlateIM/Private/Implementation/SlateIM_Widgets.cpp
@@ -91,7 +91,7 @@ namespace SlateIM
 
 		if (!EditableText)
 		{
-			FWidgetActivation ActivationData = FSlateIMManager::Get().GetWidgetActivationData();
+			TWeakPtr<FSlateIMWidgetActivationMetadata> ActivationData = Scope.GetOrCreateActivationMetadata().ToWeakPtr();
 
 			EditableText =
 				SNew(SEditableTextBox)
@@ -100,11 +100,11 @@ namespace SlateIM
 				.HintText(FText::FromStringView(HintText))
 				.OnTextChanged_Lambda([ActivationData](const FText& NewText)
 				{
-					FSlateIMManager::Get().ActivateWidget(ActivationData);
+					FSlateIMManager::Get().ActivateWidget(ActivationData.Pin());
 				})
 				.OnTextCommitted_Lambda([ActivationData](const FText& NewText, ETextCommit::Type)
 				{
-					FSlateIMManager::Get().ActivateWidget(ActivationData);
+					FSlateIMManager::Get().ActivateWidget(ActivationData.Pin());
 				});
 
 			if (TextStyle)
@@ -116,7 +116,7 @@ namespace SlateIM
 		}
 		else
 		{
-			bWasActivated = FSlateIMManager::Get().IsWidgetActivatedThisFrame();
+			bWasActivated = Scope.IsActivatedThisFrame();
 
 			const bool bIsHashDirty = Scope.IsDataHashDirty();
 
@@ -285,9 +285,9 @@ namespace SlateIM
 		{
 			ButtonWidget =
 				SNew(SImButton)
-				.OnClicked_Lambda([ActivationData = FSlateIMManager::Get().GetWidgetActivationData()]()
+				.OnClicked_Lambda([ActivationData = Scope.GetOrCreateActivationMetadata().ToWeakPtr()]()
 				{
-					FSlateIMManager::Get().ActivateWidget(ActivationData);
+					FSlateIMManager::Get().ActivateWidget(ActivationData.Pin());
 					return FReply::Handled();
 				});
 
@@ -302,7 +302,7 @@ namespace SlateIM
 		}
 		else
 		{
-			bWasClicked = FSlateIMManager::Get().IsWidgetActivatedThisFrame();
+			bWasClicked = Scope.IsActivatedThisFrame();
 
 			if (Scope.IsDataHashDirty())
 			{
@@ -345,9 +345,9 @@ namespace SlateIM
 			CheckboxWidget =
 				SNew(SImCheckBox)
 				.IsChecked(InOutCurrentState)
-				.OnCheckStateChanged_Lambda([ActivationData = FSlateIMManager::Get().GetWidgetActivationData()](ECheckBoxState NewState)
+				.OnCheckStateChanged_Lambda([ActivationData = Scope.GetOrCreateActivationMetadata().ToWeakPtr()](ECheckBoxState NewState)
 				{
-					FSlateIMManager::Get().ActivateWidget(ActivationData);
+					FSlateIMManager::Get().ActivateWidget(ActivationData.Pin());
 				});
 
 			CheckboxWidget->SetText(FText::FromStringView(InText));
@@ -361,7 +361,7 @@ namespace SlateIM
 		}
 		else
 		{
-			bValueChanged = FSlateIMManager::Get().IsWidgetActivatedThisFrame();
+			bValueChanged = Scope.IsActivatedThisFrame();
 			if (bValueChanged)
 			{
 				InOutCurrentState = CheckboxWidget->GetCheckedState();
@@ -403,7 +403,7 @@ namespace SlateIM
 		if (!SpinBoxWidget)
 		{
 			TSharedRef<FSpinBoxState> SpinBoxState = MakeShared<FSpinBoxState>();
-			FWidgetActivation ActivationData = FSlateIMManager::Get().GetWidgetActivationData();
+			TWeakPtr<FSlateIMWidgetActivationMetadata> ActivationData = Scope.GetOrCreateActivationMetadata().ToWeakPtr();
 			SpinBoxWidget =
 				SNew(SSpinBox<NumericType>)
 				.MinValue(Min)
@@ -412,15 +412,16 @@ namespace SlateIM
 				.OnBeginSliderMovement_Lambda([SpinBoxState](){ SpinBoxState->bIsChanging = true; })
 				.OnEndSliderMovement_Lambda([SpinBoxState, ActivationData](NumericType NewVal)
 				{
-					SpinBoxState->bIsChanging = false; FSlateIMManager::Get().ActivateWidget(ActivationData);
+					SpinBoxState->bIsChanging = false;
+					FSlateIMManager::Get().ActivateWidget(ActivationData.Pin());
 				})
 				.OnValueChanged_Lambda([ActivationData](NumericType NewVal)
 				{
-					FSlateIMManager::Get().ActivateWidget(ActivationData);
+					FSlateIMManager::Get().ActivateWidget(ActivationData.Pin());
 				})
 				.OnValueCommitted_Lambda([ActivationData](NumericType NewVal, ETextCommit::Type CommitType)
 				{
-					FSlateIMManager::Get().ActivateWidget(ActivationData);
+					FSlateIMManager::Get().ActivateWidget(ActivationData.Pin());
 				});
 
 			SpinBoxWidget->AddMetadata(SpinBoxState);
@@ -436,7 +437,7 @@ namespace SlateIM
 			SpinBoxWidget->SetMinValue(Min);
 			SpinBoxWidget->SetMaxValue(Max);
 
-			bValueChanged = FSlateIMManager::Get().IsWidgetActivatedThisFrame();
+			bValueChanged = Scope.IsActivatedThisFrame();
 			if (bValueChanged)
 			{
 				InOutValue = SpinBoxWidget->GetValue();
@@ -492,9 +493,9 @@ namespace SlateIM
 				.MaxValue(Max)
 				.StepSize(Step)
 				.Value(InOutValue)
-				.OnValueChanged_Lambda([ActivationData = FSlateIMManager::Get().GetWidgetActivationData()](float NewVal)
+				.OnValueChanged_Lambda([ActivationData = Scope.GetOrCreateActivationMetadata().ToWeakPtr()](float NewVal)
 				{
-					FSlateIMManager::Get().ActivateWidget(ActivationData);
+					FSlateIMManager::Get().ActivateWidget(ActivationData.Pin());
 				});
 
 			if (SliderStyle)
@@ -509,7 +510,7 @@ namespace SlateIM
 			SliderWidget->SetMinAndMaxValues(Min, Max);
 			SliderWidget->SetStepSize(Step);
 
-			bValueChanged = FSlateIMManager::Get().IsWidgetActivatedThisFrame();
+			bValueChanged = Scope.IsActivatedThisFrame();
 			if (bValueChanged)
 			{
 				InOutValue = SliderWidget->GetValue();
@@ -588,9 +589,9 @@ namespace SlateIM
 			ComboWidget =
 				SNew(STextComboBox)
 				.OptionsSource(&WidgetIMData->Data)
-				.OnSelectionChanged_Lambda([ActivationData = FSlateIMManager::Get().GetWidgetActivationData()](TSharedPtr<FString> NewVal, ESelectInfo::Type SelectType)
+				.OnSelectionChanged_Lambda([ActivationData = Scope.GetOrCreateActivationMetadata().ToWeakPtr()](TSharedPtr<FString> NewVal, ESelectInfo::Type SelectType)
 				{
-					FSlateIMManager::Get().ActivateWidget(ActivationData);
+					FSlateIMManager::Get().ActivateWidget(ActivationData.Pin());
 				});
 
 			if (InComboStyle)
@@ -615,7 +616,7 @@ namespace SlateIM
 			}
 			else
 			{
-				bValueChanged = FSlateIMManager::Get().IsWidgetActivatedThisFrame();
+				bValueChanged = Scope.IsActivatedThisFrame();
 				if (bValueChanged)
 				{
 					TSharedPtr<FString> SelectedItem = ComboWidget->GetSelectedItem();
@@ -664,10 +665,10 @@ namespace SlateIM
 				.SelectionMode(ESelectionMode::Single)
 				.ListViewStyle(InStyle)
 				.ListItemsSource(&(ListViewData->Data))
-				.OnSelectionChanged_Lambda([ActivationData = FSlateIMManager::Get().GetWidgetActivationData()](TSharedPtr<FString> NewVal, ESelectInfo::Type SelectType)
+				.OnSelectionChanged_Lambda([ActivationData = Scope.GetOrCreateActivationMetadata().ToWeakPtr()](TSharedPtr<FString> NewVal, ESelectInfo::Type SelectType)
 				{
 					UE_LOG(LogSlateIM, Verbose, TEXT("Selected %s"), NewVal.IsValid() ? **NewVal.Get() : TEXT("[NULL]"));
-					FSlateIMManager::Get().ActivateWidget(ActivationData);
+					FSlateIMManager::Get().ActivateWidget(ActivationData.Pin());
 				})
 				.OnGenerateRow_Lambda([](TSharedPtr<FString> ListItem, const TSharedRef<STableViewBase>& OwnerTable)
 				{
@@ -705,7 +706,7 @@ namespace SlateIM
 			}
 			else
 			{
-				bSelectionChanged = FSlateIMManager::Get().IsWidgetActivatedThisFrame();
+				bSelectionChanged = Scope.IsActivatedThisFrame();
 				if (bSelectionChanged)
 				{
 					TArray<TSharedPtr<FString>> SelectedItems = ListWidget->GetSelectedItems();
diff --git a/Engine/Plugins/Experimental/SlateIM/Source/SlateIM/Private/Misc/SlateIMManager.cpp b/Engine/Plugins/Experimental/SlateIM/Source/SlateIM/Private/Misc/SlateIMManager.cpp
--- a/Engine/Plugins/Experimental/SlateIM/Source/SlateIM/Private/Misc/SlateIMManager.cpp
+++ b/Engine/Plugins/Experimental/SlateIM/Source/SlateIM/Private/Misc/SlateIMManager.cpp
@@ -9,6 +9,7 @@
 #include "Roots/ISlateIMRoot.h"
 #include "SlateIMLogging.h"
 #include "SlateIM.h"
+#include "SlateIMWidgetActivationMetadata.h"
 
 namespace SlateIM
 {
@@ -290,29 +291,30 @@ namespace SlateIM
 		++CurrentRoot->CurrentContainerStack.Top().LastUsedChildIndex;
 	}
 
-	void FSlateIMManager::ActivateWidget(const FWidgetActivation& ActivationData)
+	void FSlateIMManager::ActivateWidget(const TSharedPtr<FSlateIMWidgetActivationMetadata>& ActivationData)
 	{
-		// Setting values while building immediate mode ui causes some widgets to activate
-		// This is undesirable as it can cause feedback loops
-		if (!bBuildingImmediateModeWidgets)
+		if (ActivationData.IsValid())
 		{
-			UE_LOG(LogSlateIM, Verbose, TEXT("ActivateWidget | Root: [%s] | Container [%d] | Widget [%d]"), *ActivationData.RootName.ToString(), ActivationData.ContainerIndex, ActivationData.WidgetIndex);
-			ActivatedWidgets.Add(ActivationData);
+			// Setting values while building immediate mode ui causes some widgets to activate
+			// This is undesirable as it can cause feedback loops
+			if (!bBuildingImmediateModeWidgets)
+			{
+				FWidgetActivation WidgetActivation = FWidgetActivation{ ActivationData->RootName, ActivationData->ContainerIndex, ActivationData->WidgetIndex };
+				UE_LOG(LogSlateIM, Verbose, TEXT("ActivateWidget | Root: [%s] | Container [%d] | Widget [%d]"), *WidgetActivation.RootName.ToString(), WidgetActivation.ContainerIndex, WidgetActivation.WidgetIndex);
+				ActivatedWidgets.Add(MoveTemp(WidgetActivation));
+			}
 		}
 	}
 
-	bool FSlateIMManager::IsWidgetActivatedThisFrame() const
+	bool FSlateIMManager::IsWidgetActivatedThisFrame(const TSharedPtr<FSlateIMWidgetActivationMetadata>& ActivationData) const
 	{
-		return ActivatedWidgets.Contains(GetWidgetActivationData());
-	}
-
-	FWidgetActivation FSlateIMManager::GetWidgetActivationData() const
-	{
-		const FRootNode& Root = FSlateIMManager::Get().GetCurrentRoot();
-		const int32 ContainerIndex = Root.CurrentContainerStack.Num() - 1;
-		const int32 WidgetIndex = Root.CurrentWidgetIndex;
+		if (ActivationData.IsValid())
+		{
+			const FWidgetActivation WidgetActivation = FWidgetActivation{ ActivationData->RootName, ActivationData->ContainerIndex, ActivationData->WidgetIndex };
+			return ActivatedWidgets.Contains(WidgetActivation);
+		}
 
-		return { Root.RootName, ContainerIndex, WidgetIndex };
+		return false;
 	}
 
 	void FSlateIMManager::ResetAlignmentData()
diff --git a/Engine/Plugins/Experimental/SlateIM/Source/SlateIM/Private/Misc/SlateIMManager.h b/Engine/Plugins/Experimental/SlateIM/Source/SlateIM/Private/Misc/SlateIMManager.h
--- a/Engine/Plugins/Experimental/SlateIM/Source/SlateIM/Private/Misc/SlateIMManager.h
+++ b/Engine/Plugins/Experimental/SlateIM/Source/SlateIM/Private/Misc/SlateIMManager.h
@@ -17,6 +17,7 @@
 #include "SlateIMExamples.h"
 #endif
 
+struct FSlateIMWidgetActivationMetadata;
 class FMemMark;
 class ISlateIMChild;
 class ISlateIMRoot;
@@ -149,9 +150,8 @@ namespace SlateIM
 		void UpdateCurrentChild(FSlateIMChild Child, const FSlateIMSlotData& AlignmentData);
 		void AdvanceToNextWidget();
 
-		void ActivateWidget(const FWidgetActivation& ActivationData);
-		bool IsWidgetActivatedThisFrame() const;
-		FWidgetActivation GetWidgetActivationData() const;
+		void ActivateWidget(const TSharedPtr<FSlateIMWidgetActivationMetadata>& ActivationData);
+		bool IsWidgetActivatedThisFrame(const TSharedPtr<FSlateIMWidgetActivationMetadata>& ActivationData) const;
 
 		FSlateIMSlotData GetCurrentAlignmentData(
 			const FMargin& DefaultPadding,
diff --git a/Engine/Plugins/Experimental/SlateIM/Source/SlateIM/Private/Misc/SlateIMWidgetActivationMetadata.h b/Engine/Plugins/Experimental/SlateIM/Source/SlateIM/Private/Misc/SlateIMWidgetActivationMetadata.h
new file mode 100644
--- /dev/null
+++ b/Engine/Plugins/Experimental/SlateIM/Source/SlateIM/Private/Misc/SlateIMWidgetActivationMetadata.h
@@ -0,0 +1,20 @@
+﻿// Copyright Epic Games, Inc. All Rights Reserved.
+
+#pragma once
+
+#include "Types/ISlateMetaData.h"
+
+struct FSlateIMWidgetActivationMetadata : ISlateMetaData
+{
+	SLATE_METADATA_TYPE(FSlateIMWidgetActivationMetadata, ISlateMetaData);
+
+	FSlateIMWidgetActivationMetadata(const FName& RootName, int32 ContainerIndex, int32 WidgetIndex)
+		: RootName(RootName)
+		, ContainerIndex(ContainerIndex)
+		, WidgetIndex(WidgetIndex)	
+	{}
+	
+	FName RootName = NAME_None;
+	int32 ContainerIndex = INDEX_NONE;
+	int32 WidgetIndex = INDEX_NONE;
+};
\ No newline at end of file
diff --git a/Engine/Plugins/Experimental/SlateIM/Source/SlateIM/Private/Misc/SlateIMWidgetScope.h b/Engine/Plugins/Experimental/SlateIM/Source/SlateIM/Private/Misc/SlateIMWidgetScope.h
--- a/Engine/Plugins/Experimental/SlateIM/Source/SlateIM/Private/Misc/SlateIMWidgetScope.h
+++ b/Engine/Plugins/Experimental/SlateIM/Source/SlateIM/Private/Misc/SlateIMWidgetScope.h
@@ -3,10 +3,11 @@
 #pragma once
 
 #include "SlateIMManager.h"
-	
+#include "SlateIMWidgetActivationMetadata.h"
+
 namespace SlateIM
 {
-	template<typename WidgetType>
+	template<typename WidgetType, bool CanHaveMetadata = TIsDerivedFrom<WidgetType, SWidget>::Value>
 	struct FWidgetScope
 	{
 		FWidgetScope(const FMargin DefaultPadding = Defaults::Padding,
@@ -23,6 +24,14 @@ namespace SlateIM
 			, bIsSlotDirty(false)
 		{
 			HashStringView(FSlateIMManager::Get().GetCurrentRoot().CurrentToolTip);
+
+			if constexpr (CanHaveMetadata)
+			{
+				if (Widget.IsValid())
+				{
+					ActivationMetadata = Widget->template GetMetaData<FSlateIMWidgetActivationMetadata>();
+				}
+			}
 		}
 
 		FWidgetScope(const bool bAutoSize)
@@ -45,10 +54,30 @@ namespace SlateIM
 			, bIsSlotDirty(false)
 		{
 			HashStringView(FSlateIMManager::Get().GetCurrentRoot().CurrentToolTip);
+
+			if constexpr (CanHaveMetadata)
+			{
+				if (Widget.IsValid())
+				{
+					ActivationMetadata = Widget->template GetMetaData<FSlateIMWidgetActivationMetadata>();
+				}
+			}
 		}
 
 		~FWidgetScope()
 		{
+			if constexpr (CanHaveMetadata)
+			{
+				// Update the activation metadata to match the current frame's values
+				if (ActivationMetadata.IsValid())
+				{
+					const FRootNode& Root = FSlateIMManager::Get().GetCurrentRoot();
+					ActivationMetadata->RootName = Root.RootName;
+					ActivationMetadata->ContainerIndex = Root.CurrentContainerStack.Num() - 1;
+					ActivationMetadata->WidgetIndex = Root.CurrentWidgetIndex;
+				}
+			}
+
 			FRootNode& CurrentRoot = FSlateIMManager::Get().GetMutableCurrentRoot();
 
 			FWidgetHash CurrentHash = CurrentRoot.GetWidgetHash();
@@ -101,6 +130,14 @@ namespace SlateIM
 		{
 			Widget = NewWidget;
 			bIsSlotDirty = true;
+			
+			if constexpr (CanHaveMetadata)
+			{
+				if (ActivationMetadata.IsValid() && Widget.IsValid() && !Widget->template GetMetaData<FSlateIMWidgetActivationMetadata>().IsValid())
+				{
+					Widget->AddMetadata(ActivationMetadata.ToSharedRef());
+				}
+			}
 		}
 
 		template<typename T>
@@ -120,8 +157,43 @@ namespace SlateIM
 			return FSlateIMManager::Get().GetCurrentRoot().GetWidgetHash().DataHash != DataHash;
 		}
 
+		bool IsActivatedThisFrame() const
+		{
+			return FSlateIMManager::Get().IsWidgetActivatedThisFrame(ActivationMetadata);
+		}
+
+		TSharedPtr<FSlateIMWidgetActivationMetadata> GetOrCreateActivationMetadata()
+		{
+			if constexpr (CanHaveMetadata)
+			{
+				if (!ActivationMetadata.IsValid())
+				{
+					if (Widget.IsValid())
+					{
+						ActivationMetadata = Widget->template GetMetaData<FSlateIMWidgetActivationMetadata>();
+					}
+					
+					if (!ActivationMetadata.IsValid())
+					{
+						const FRootNode& Root = FSlateIMManager::Get().GetCurrentRoot();
+						const int32 ContainerIndex = Root.CurrentContainerStack.Num() - 1;
+						const int32 WidgetIndex = Root.CurrentWidgetIndex;
+						ActivationMetadata = MakeShared<FSlateIMWidgetActivationMetadata>(Root.RootName, ContainerIndex, WidgetIndex);
+					}
+					
+					if (Widget.IsValid())
+					{
+						Widget->AddMetadata(ActivationMetadata.ToSharedRef());
+					}
+				}
+			}
+			
+			return ActivationMetadata;
+		}
+
 	private:
 		TSharedPtr<WidgetType> Widget;
+		TSharedPtr<FSlateIMWidgetActivationMetadata> ActivationMetadata;
 		FSlateIMSlotData AlignmentData;
 		FXxHash64Builder HashBuilder;
 		mutable FXxHash64 DataHash;
