SlateIM: Expose setting the TableRowStyle for BeginTable()

This allows for customizing table row styles, such as alternating row background brushes.

#rnx

diff --git a/Engine/Plugins/Experimental/SlateIM/Source/SlateIM/Private/Containers/SImTableView.cpp b/Engine/Plugins/Experimental/SlateIM/Source/SlateIM/Private/Containers/SImTableView.cpp
--- a/Engine/Plugins/Experimental/SlateIM/Source/SlateIM/Private/Containers/SImTableView.cpp
+++ b/Engine/Plugins/Experimental/SlateIM/Source/SlateIM/Private/Containers/SImTableView.cpp
@@ -180,10 +180,10 @@ void SImTableRow::PrivateRegisterAttributes(FSlateAttributeInitializer&)
 {
 }
 
-void SImTableRow::Construct(const FArguments& InArgs, const TSharedRef<SImTableView>& InOwnerTableView, const TSharedRef<FSlateIMTableRow>& InTableRow)
+void SImTableRow::Construct(const FSuperRowType::FArguments& InArgs, const TSharedRef<SImTableView>& InOwnerTableView, const TSharedRef<FSlateIMTableRow>& InTableRow)
 {
 	TableRow = InTableRow;
-	SMultiColumnTableRow<TSharedRef<FSlateIMTableRow>>::Construct(FSuperRowType::FArguments(), InOwnerTableView);
+	SMultiColumnTableRow<TSharedRef<FSlateIMTableRow>>::Construct(InArgs, InOwnerTableView);
 }
 
 TSharedRef<SWidget> SImTableRow::GenerateWidgetForColumn(const FName& ColumnName)
@@ -355,6 +355,11 @@ void SImTableView::UpdateColumns()
 	}
 }
 
+void SImTableView::SetTableRowStyle(const FTableRowStyle* InRowStyle)
+{
+	RowStyle = InRowStyle;
+}
+
 bool SImTableView::IsTree() const
 {
 	for (const TSharedRef<FSlateIMTableRow>& TableRow : TableRows)
@@ -370,7 +375,8 @@ bool SImTableView::IsTree() const
 
 TSharedRef<ITableRow> SImTableView::GenerateRow(TSharedRef<FSlateIMTableRow> InTableRow, const TSharedRef<STableViewBase>& OwnerTable)
 {
-	return SNew(SImTableRow, SharedThis(this), InTableRow);
+	return SNew(SImTableRow, SharedThis(this), InTableRow)
+		.Style(RowStyle);
 }
 
 void SImTableView::GatherChildren(TSharedRef<FSlateIMTableRow> Row, TArray<TSharedRef<FSlateIMTableRow>>& OutChildren) const
diff --git a/Engine/Plugins/Experimental/SlateIM/Source/SlateIM/Private/Containers/SImTableView.h b/Engine/Plugins/Experimental/SlateIM/Source/SlateIM/Private/Containers/SImTableView.h
--- a/Engine/Plugins/Experimental/SlateIM/Source/SlateIM/Private/Containers/SImTableView.h
+++ b/Engine/Plugins/Experimental/SlateIM/Source/SlateIM/Private/Containers/SImTableView.h
@@ -51,11 +51,8 @@ class SImTableRow : public SMultiColumnTableRow<TSharedRef<FSlateIMTableRow>>
 {
 	SLATE_DECLARE_WIDGET(SImTableRow, SMultiColumnTableRow<TSharedRef<FSlateIMTableRow>>)
 	
-public:
-	SLATE_BEGIN_ARGS(SImTableRow) { }
-	SLATE_END_ARGS()
-	
-	void Construct(const FArguments& InArgs, const TSharedRef<SImTableView>& InOwnerTableView, const TSharedRef<FSlateIMTableRow>& InTableRow);
+public:	
+	void Construct(const FSuperRowType::FArguments& InArgs, const TSharedRef<SImTableView>& InOwnerTableView, const TSharedRef<FSlateIMTableRow>& InTableRow);
 	
 	virtual TSharedRef<SWidget> GenerateWidgetForColumn(const FName& ColumnName) override;
 	
@@ -93,6 +90,8 @@ public:
 
 	void UpdateColumns();
 
+	void SetTableRowStyle(const FTableRowStyle* InRowStyle);
+
 	// A table is a tree if any of its rows has children
 	bool IsTree() const;
 
@@ -113,6 +112,7 @@ private:
 	TSharedPtr<SImTableHeader> Header;
 	int32 ColumnCount = 0;
 	bool bIsDirty = false;
+	const FTableRowStyle* RowStyle = nullptr;
 
 	mutable int32 CachedNumLiveWidgets = 0;
 };
diff --git a/Engine/Plugins/Experimental/SlateIM/Source/SlateIM/Private/Implementation/SlateIM_Containers.cpp b/Engine/Plugins/Experimental/SlateIM/Source/SlateIM/Private/Implementation/SlateIM_Containers.cpp
--- a/Engine/Plugins/Experimental/SlateIM/Source/SlateIM/Private/Implementation/SlateIM_Containers.cpp
+++ b/Engine/Plugins/Experimental/SlateIM/Source/SlateIM/Private/Implementation/SlateIM_Containers.cpp
@@ -250,7 +250,7 @@ namespace SlateIM
 		FSlateIMManager::Get().PopContainer<FSlateIMTableRow>();
 	}
 	
-	void BeginTable(const FTableViewStyle* InStyle)
+	void BeginTable(const FTableViewStyle* InStyle, const FTableRowStyle* InRowStyle)
 	{
 		TSharedPtr<SImTableView> ContainerWidget;
 
@@ -273,6 +273,7 @@ namespace SlateIM
 				ContainerWidget->SetStyle(InStyle);
 			}
 
+			ContainerWidget->SetTableRowStyle(InRowStyle);
 			ContainerWidget->BeginTableUpdates();
 		}
 
diff --git a/Engine/Plugins/Experimental/SlateIM/Source/SlateIM/Public/SlateIM.h b/Engine/Plugins/Experimental/SlateIM/Source/SlateIM/Public/SlateIM.h
--- a/Engine/Plugins/Experimental/SlateIM/Source/SlateIM/Public/SlateIM.h
+++ b/Engine/Plugins/Experimental/SlateIM/Source/SlateIM/Public/SlateIM.h
@@ -29,6 +29,7 @@ struct FProgressBarStyle;
 struct FSlateBrush;
 struct FSliderStyle;
 struct FSpinBoxStyle;
+struct FTableRowStyle;
 struct FTableViewStyle;
 struct FTextBlockStyle;
 
@@ -176,8 +177,9 @@ namespace SlateIM
 	 * 
 	 * @param InStyle (optional) A Table View Style to override the default style. Changes to the style after the first invocation are not reflected,
 	 *				  provide a new style object to update the visuals at runtime
+	 * @param InRowStyle (optional) A Table Row Style to override the default table row style. Changing the style only affects new table rows.
 	 */
-	SLATEIM_API void BeginTable(const FTableViewStyle* InStyle = nullptr);
+	SLATEIM_API void BeginTable(const FTableViewStyle* InStyle = nullptr, const FTableRowStyle* InRowStyle = nullptr);
 	/**
 	 * End the current Table container.
 	 */
