SlateIM: Context Menus now only open with Right Mouse Button clicks

#rnx

diff --git a/Engine/Plugins/Experimental/SlateIM/Source/SlateIM/Private/Containers/SImContextMenuAnchor.cpp b/Engine/Plugins/Experimental/SlateIM/Source/SlateIM/Private/Containers/SImContextMenuAnchor.cpp
--- a/Engine/Plugins/Experimental/SlateIM/Source/SlateIM/Private/Containers/SImContextMenuAnchor.cpp
+++ b/Engine/Plugins/Experimental/SlateIM/Source/SlateIM/Private/Containers/SImContextMenuAnchor.cpp
@@ -152,7 +152,7 @@ void SImContextMenuAnchor::EndSubMenu()
 
 FReply SImContextMenuAnchor::OnMouseButtonUp(const FGeometry& MyGeometry, const FPointerEvent& MouseEvent)
 {
-	if (MenuWidget)
+	if (MenuWidget && MouseEvent.GetEffectingButton() == EKeys::RightMouseButton)
 	{
 		FWidgetPath WidgetPath = MouseEvent.GetEventPath() != nullptr ? *MouseEvent.GetEventPath() : FWidgetPath();
 		OpenedMenu = FSlateApplication::Get().PushMenu(AsShared(), WidgetPath, MenuWidget.ToSharedRef(), MouseEvent.GetScreenSpacePosition(), FPopupTransitionEffect(FPopupTransitionEffect::ContextMenu));
