SlateIM: Fix ensure with empty tables and fix crash when not specifying a table row style

Also added an example of a table with alternating row backgrounds

#rnx

diff --git a/Engine/Plugins/Experimental/SlateIM/Source/SlateIM/Private/Containers/SImTableView.cpp b/Engine/Plugins/Experimental/SlateIM/Source/SlateIM/Private/Containers/SImTableView.cpp
--- a/Engine/Plugins/Experimental/SlateIM/Source/SlateIM/Private/Containers/SImTableView.cpp
+++ b/Engine/Plugins/Experimental/SlateIM/Source/SlateIM/Private/Containers/SImTableView.cpp
@@ -376,7 +376,7 @@ bool SImTableView::IsTree() const
 TSharedRef<ITableRow> SImTableView::GenerateRow(TSharedRef<FSlateIMTableRow> InTableRow, const TSharedRef<STableViewBase>& OwnerTable)
 {
 	return SNew(SImTableRow, SharedThis(this), InTableRow)
-		.Style(RowStyle);
+		.Style(RowStyle ? RowStyle : &FCoreStyle::Get().GetWidgetStyle<FTableRowStyle>("TableView.Row"));
 }
 
 void SImTableView::GatherChildren(TSharedRef<FSlateIMTableRow> Row, TArray<TSharedRef<FSlateIMTableRow>>& OutChildren) const
diff --git a/Engine/Plugins/Experimental/SlateIM/Source/SlateIM/Private/Implementation/SlateIM_Containers.cpp b/Engine/Plugins/Experimental/SlateIM/Source/SlateIM/Private/Implementation/SlateIM_Containers.cpp
--- a/Engine/Plugins/Experimental/SlateIM/Source/SlateIM/Private/Implementation/SlateIM_Containers.cpp
+++ b/Engine/Plugins/Experimental/SlateIM/Source/SlateIM/Private/Implementation/SlateIM_Containers.cpp
@@ -289,13 +289,20 @@ namespace SlateIM
 		}
 
 		TSharedPtr<FSlateIMTableRow> Row = FSlateIMManager::Get().GetCurrentContainer<FSlateIMTableRow>();
-		if (ensure(Row))
+		TSharedPtr<SImTableView> Table = FSlateIMManager::Get().GetCurrentContainer<SImTableView>();
+		ensureAlwaysMsgf((Row.IsValid() || Table.IsValid()), TEXT("Current container should be a row or a table - Is there a missing SlateIM::EndX() statement?"));
+		if (Row)
 		{
 			EndTableRow();
 		}
+
+		// If we just ended a Row, then we need to fetch the table
+		if (!Table.IsValid())
+		{
+			Table = FSlateIMManager::Get().GetCurrentContainer<SImTableView>();
+		}
 		
-		TSharedPtr<SImTableView> Table = FSlateIMManager::Get().GetCurrentContainer<SImTableView>();
-		if (ensure(Table))
+		if (ensureAlwaysMsgf(Table, TEXT("Current container should be a table - Is there a missing SlateIM::EndX() statement?")))
 		{
 			FSlateIMManager::Get().PopContainer<SImTableView>();
 			Table->EndTableUpdates();
diff --git a/Engine/Plugins/Experimental/SlateIM/Source/SlateIM/Private/Misc/SlateIMExamples.cpp b/Engine/Plugins/Experimental/SlateIM/Source/SlateIM/Private/Misc/SlateIMExamples.cpp
--- a/Engine/Plugins/Experimental/SlateIM/Source/SlateIM/Private/Misc/SlateIMExamples.cpp
+++ b/Engine/Plugins/Experimental/SlateIM/Source/SlateIM/Private/Misc/SlateIMExamples.cpp
@@ -605,11 +605,12 @@ void FSlateIMTestWidget::Draw()
 		// Tree Example
 		{
 			SCOPED_NAMED_EVENT_TEXT("Tree Example", FColorList::Goldenrod);
+			const FTableRowStyle* TableRowStyle = &FAppStyle::Get().GetWidgetStyle<FTableRowStyle>("TableView.AlternatingRow");
 			SlateIM::MinWidth(500.f);
 			SlateIM::MinHeight(200.f);
 			SlateIM::MaxHeight(200.f);
 			SlateIM::VAlign(VAlign_Fill);
-			SlateIM::BeginTable();
+			SlateIM::BeginTable(nullptr, TableRowStyle);
 			SlateIM::AddTableColumn(TEXT("Name"));
 			SlateIM::AddTableColumn(TEXT("Type"));
 			{
