SlateIM: Add support for Tab widgets

diff --git a/Engine/Plugins/Experimental/SlateIM/Source/SlateIM/Private/Containers/SImTab.cpp b/Engine/Plugins/Experimental/SlateIM/Source/SlateIM/Private/Containers/SImTab.cpp
new file mode 100644
--- /dev/null
+++ b/Engine/Plugins/Experimental/SlateIM/Source/SlateIM/Private/Containers/SImTab.cpp
@@ -0,0 +1,64 @@
+﻿// Copyright Epic Games, Inc. All Rights Reserved.
+
+
+#include "SImTab.h"
+
+#include "SImStackBox.h"
+
+SLATE_IMPLEMENT_WIDGET(SImTab)
+
+void SImTab::PrivateRegisterAttributes(FSlateAttributeInitializer&)
+{
+}
+
+int32 SImTab::GetNumChildren()
+{
+	if (Container)
+	{
+		return Container->GetNumChildren();
+	}
+	
+	return GetChildren()->Num();
+}
+
+FSlateIMChild SImTab::GetChild(int32 Index)
+{
+	if (Container)
+	{
+		return Container->GetChild(Index);
+	}
+	
+	if (Index >= 0 && Index < GetChildren()->Num())
+	{
+		return GetChildren()->GetChildAt(Index);
+	}
+
+	return nullptr;
+}
+
+void SImTab::UpdateChild(FSlateIMChild Child, int32 Index, const FSlateIMSlotData& AlignmentData)
+{
+	if (!Container)
+	{
+		Container = SNew(SImStackBox)
+			.Orientation(Orient_Vertical);
+
+		SetContent(Container.ToSharedRef());
+	}
+	
+	Container->UpdateChild(Child, Index, AlignmentData);
+}
+
+void SImTab::RemoveUnusedChildren(int32 LastUsedChildIndex)
+{
+	if (Container)
+	{
+		Container->RemoveUnusedChildren(LastUsedChildIndex);
+	}
+}
+
+void SImTab::ForceCloseTab()
+{
+	bCanEverClose = true;
+	RequestCloseTab();
+}
diff --git a/Engine/Plugins/Experimental/SlateIM/Source/SlateIM/Private/Containers/SImTab.h b/Engine/Plugins/Experimental/SlateIM/Source/SlateIM/Private/Containers/SImTab.h
new file mode 100644
--- /dev/null
+++ b/Engine/Plugins/Experimental/SlateIM/Source/SlateIM/Private/Containers/SImTab.h
@@ -0,0 +1,34 @@
+﻿// Copyright Epic Games, Inc. All Rights Reserved.
+
+#pragma once
+
+#include "ISlateIMContainer.h"
+#include "Widgets/Docking/SDockTab.h"
+
+class SImStackBox;
+
+class SLATEIM_API SImTab : public SDockTab, public ISlateIMContainer
+{
+	SLATE_DECLARE_WIDGET(SImTab, SDockTab)
+	SLATE_IM_TYPE_DATA(SImTab, ISlateIMContainer)
+	
+public:
+	virtual int32 GetNumChildren() override;
+	virtual FSlateIMChild GetChild(int32 Index) override;
+	virtual void UpdateChild(FSlateIMChild Child, int32 Index, const FSlateIMSlotData& AlignmentData) override;
+	virtual void RemoveUnusedChildren(int32 LastUsedChildIndex) override;
+
+	virtual FSlateIMChild GetContainer() override
+	{
+		return AsShared();
+	}
+
+	void ForceCloseTab();
+
+	FName TabId = NAME_None;
+	FSlateIcon TabIcon;
+	FText TabTitle = FText::GetEmpty();
+
+private:
+	TSharedPtr<SImStackBox> Container;
+};
diff --git a/Engine/Plugins/Experimental/SlateIM/Source/SlateIM/Private/Containers/SImTabGroup.cpp b/Engine/Plugins/Experimental/SlateIM/Source/SlateIM/Private/Containers/SImTabGroup.cpp
new file mode 100644
--- /dev/null
+++ b/Engine/Plugins/Experimental/SlateIM/Source/SlateIM/Private/Containers/SImTabGroup.cpp
@@ -0,0 +1,420 @@
+﻿// Copyright Epic Games, Inc. All Rights Reserved.
+
+
+#include "SImTabGroup.h"
+
+#include "Misc/SlateIMManager.h"
+#include "Roots/SlateIMWindowRoot.h"
+#include "SImTab.h"
+#include "Widgets/Docking/SDockTab.h"
+#include "Widgets/Text/STextBlock.h"
+
+class FSlateImTabStack : public FTabManager::FStack
+{
+public:
+	bool DoesContainTab(const FName& TabId) const
+	{
+		return Tabs.ContainsByPredicate([TabId](const FTabManager::FTab& Tab)
+		{
+			return Tab.TabId.TabType == TabId;
+		});
+	}
+
+	void SetForegroundTabIfUnset(const FName& TabId)
+	{
+		if (ForegroundTabId.TabType.IsNone())
+		{
+			SetForegroundTab(TabId);
+		}
+	}
+	
+	bool IsTabForegrounded(const FName& TabId) const
+	{
+		return ForegroundTabId.TabType == TabId;
+	}
+};
+
+class FSlateImTabSplitter : public FTabManager::FSplitter
+{
+public:
+	TSharedPtr<FTabManager::FStack> GetChildAsStack(int32 ChildIndex) const
+	{
+		if (ChildNodes.IsValidIndex(ChildIndex))
+		{
+			return ChildNodes[ChildIndex]->AsStack();
+		}
+
+		return nullptr;
+	}
+
+	TSharedPtr<FTabManager::FSplitter> GetChildAsSplitter(int32 ChildIndex) const
+	{
+		if (ChildNodes.IsValidIndex(ChildIndex))
+		{
+			return ChildNodes[ChildIndex]->AsSplitter();
+		}
+
+		return nullptr;
+	}
+	
+	bool IsTabForegrounded(const FName& TabId) const
+	{
+		for (const TSharedRef<FTabManager::FLayoutNode>& ChildNode : ChildNodes)
+		{
+			if (const TSharedPtr<FSlateImTabSplitter>& ChildSplitter = StaticCastSharedPtr<FSlateImTabSplitter>(ChildNode->AsSplitter()))
+			{
+				if (ChildSplitter->IsTabForegrounded(TabId))
+				{
+					return true;
+				}
+			}
+			else if (const TSharedPtr<FSlateImTabStack>& ChildStack = StaticCastSharedPtr<FSlateImTabStack>(ChildNode->AsStack()))
+			{
+				if (ChildStack->IsTabForegrounded(TabId))
+				{
+					return true;
+				}
+			}
+		}
+
+		return false;
+	}
+
+	bool SplitAt(int32 Index, TSharedRef<FLayoutNode> InNode)
+	{
+		if (!ChildNodes.IsValidIndex(Index) || ChildNodes[Index] != InNode)
+		{
+			ChildNodes.Insert(InNode, FMath::Min(Index, ChildNodes.Num()));
+			return true;
+		}
+
+		return false;
+	}
+
+	void RemoveChildAt(int32 Index)
+	{
+		if (ChildNodes.IsValidIndex(Index))
+		{
+			ChildNodes.RemoveAt(Index);
+		}
+	}
+
+	bool RemoveUnusedChildren(int32 StartingIndex)
+	{
+		if (ChildNodes.Num() > StartingIndex)
+		{
+			ChildNodes.SetNumUninitialized(StartingIndex, EAllowShrinking::No);
+			return true;
+		}
+
+		return false;
+	}
+};
+
+class FSlateImTabLayout : public FTabManager::FLayout
+{
+public:
+	bool IsTabForegrounded(const FName& TabId) const
+	{
+		for (const TSharedRef<FTabManager::FArea>& Area : Areas)
+		{
+			if (StaticCastSharedPtr<FSlateImTabSplitter>(Area->AsSplitter())->IsTabForegrounded(TabId))
+			{
+				return true;
+			}
+		}
+
+		return false;
+	}
+};
+
+SLATE_IMPLEMENT_WIDGET(SImTabGroup)
+
+void SImTabGroup::PrivateRegisterAttributes(FSlateAttributeInitializer&)
+{
+}
+
+void SImTabGroup::Construct(const FArguments& InArgs)
+{
+	TabGroupId = InArgs._TabGroupId;
+	const TSharedPtr<SDockTab> ParentTab = SlateIM::FSlateIMManager::Get().FindMostRecentContainer<SImTab>();
+	TabManager = FGlobalTabmanager::Get()->NewTabManager(ParentTab.IsValid() ? ParentTab.ToSharedRef() : SNew(SDockTab));
+	const TSharedRef<FTabManager::FArea> PrimaryArea = FTabManager::NewPrimaryArea();
+	WorkingLayoutNodeStack.Emplace(PrimaryArea);
+	TabLayout = FTabManager::NewLayout(TabGroupId)->AddArea(PrimaryArea);
+}
+
+SImTabGroup::~SImTabGroup()
+{
+	// Force-close the tabs to handle any that may have been pulled into their own windows
+	for (const FImTab& Tab : Tabs)
+	{
+		if (Tab.TabWidget.IsValid())
+		{
+			Tab.TabWidget->ForceCloseTab();
+		}
+	}
+}
+
+void SImTabGroup::BeginTabGroupUpdates()
+{
+	checkf(WorkingLayoutNodeStack.Num() == 1, TEXT("The Tab Group's WorkingLayoutNodeStack should have exactly 1 entry. Are your Begin and End function calls mismatched?"));
+	WorkingLayoutNodeStack.SetNumUninitialized(1, EAllowShrinking::No);
+	WorkingLayoutNodeStack[0].CurrentChildIndex = 0;
+}
+
+void SImTabGroup::FinalizeTabGroup()
+{
+	if (!bIsLayoutUpToDate)
+	{
+		check(TabLayout.IsValid());
+		bIsLayoutUpToDate = true;
+		
+		TSharedPtr<SWindow> TabWindow;
+		const TSharedPtr<ISlateIMRoot>& RootWidget = SlateIM::FSlateIMManager::Get().GetCurrentRoot().RootWidget;
+		if (RootWidget.IsValid() && RootWidget->IsA<FSlateIMWindowRoot>())
+		{
+			TabWindow = StaticCastSharedPtr<FSlateIMWindowRoot>(RootWidget)->GetWindow();
+		}
+
+		// TODO - Save layout to restore from (part of overall SlateIM save/restore user's UI state)
+		ChildSlot
+		[
+			TabManager->RestoreFrom(TabLayout.ToSharedRef(), TabWindow).ToSharedRef()
+		];
+	}
+}
+
+int32 SImTabGroup::GetNumChildren()
+{
+	return Tabs.Num();
+}
+
+FSlateIMChild SImTabGroup::GetChild(int32 Index)
+{
+	return (Tabs.IsValidIndex(Index) && Tabs[Index].TabWidget.IsValid()) ? FSlateIMChild(Tabs[Index].TabWidget.ToSharedRef()) : FSlateIMChild();
+}
+
+void SImTabGroup::UpdateChild(FSlateIMChild Child, int32 Index, const FSlateIMSlotData& AlignmentData)
+{
+	TSharedPtr<SImTab> TabWidget = Child.GetWidget<SImTab>();
+	if (!ensureMsgf(TabWidget.IsValid(), TEXT("Attempting to add a non-tab widget to a Tab Group, only Tabs can be added to Tab Groups")))
+	{
+		return;
+	}
+	
+	if (!ensureMsgf(!WorkingLayoutNodeStack.IsEmpty(), TEXT("The Tab Group's WorkingLayoutNodeStack should not be empty. Are your Begin and End function calls mismatched?")))
+	{
+		return;
+	}
+	
+	const TSharedPtr<FSlateImTabStack> Stack = StaticCastSharedPtr<FSlateImTabStack>(WorkingLayoutNodeStack.Last().Node->AsStack());
+	if (!ensureMsgf(Stack.IsValid(), TEXT("Cannot add a tab without a Tab Stack. Are your Begin and End function calls mismatched?")))
+	{
+		return;
+	}
+
+	if (!Stack->DoesContainTab(TabWidget->TabId))
+	{
+		Stack->AddTab(TabWidget->TabId, ETabState::OpenedTab);
+		Stack->SetForegroundTabIfUnset(TabWidget->TabId);
+		bIsLayoutUpToDate = false;
+	}
+
+	if (!Tabs.IsValidIndex(Index))
+	{
+		Tabs.SetNum(Index + 1);
+		bIsLayoutUpToDate = false;
+	}
+
+	if (Tabs[Index].TabId != TabWidget->TabId)
+	{
+		Tabs[Index].TabId = TabWidget->TabId;
+		bIsLayoutUpToDate = false;
+	}
+			
+	if (Tabs[Index].TabWidget != TabWidget)
+	{
+		Tabs[Index].TabWidget = TabWidget;
+		bIsLayoutUpToDate = false;
+	}
+
+	check(TabManager.IsValid());
+	if (!TabManager->HasTabSpawner(TabWidget->TabId))
+	{
+		FTabSpawnerEntry& Spawner = TabManager->RegisterTabSpawner(TabWidget->TabId, FOnSpawnTab::CreateSP(this, &SImTabGroup::SpawnTab));
+		Spawner.SetReuseTabMethod(FOnFindTabToReuse::CreateSP(this, &SImTabGroup::ReuseTab));
+		Spawner.SetDisplayName(TabWidget->TabTitle);
+		Spawner.SetIcon(TabWidget->TabIcon);
+		bIsLayoutUpToDate = false;
+	}
+}
+
+void SImTabGroup::RemoveUnusedChildren(int32 LastUsedChildIndex)
+{
+	check(TabManager.IsValid());
+	for (int32 ChildIndex = LastUsedChildIndex + 1; ChildIndex < Tabs.Num(); ++ChildIndex)
+	{
+		TabManager->UnregisterTabSpawner(Tabs[ChildIndex].TabId);
+		bIsLayoutUpToDate = false;
+	}
+
+	Tabs.SetNum(LastUsedChildIndex + 1);
+}
+
+FName SImTabGroup::GetTabGroupId() const
+{
+	return TabLayout.IsValid() ? TabLayout->GetLayoutName() : NAME_None;
+}
+
+bool SImTabGroup::IsTabForegrounded(const FName& TabId) const
+{
+	// If the Tab Manager is going to be restoring from layout this frame, query the layout directly
+	if (!bIsLayoutUpToDate)
+	{
+		check(TabLayout.IsValid());
+		return StaticCastSharedPtr<FSlateImTabLayout>(TabLayout)->IsTabForegrounded(TabId);
+	}
+	else
+	{
+		for (const FImTab& Tab : Tabs)
+		{
+			if (Tab.TabWidget.IsValid() && Tab.TabId == TabId)
+			{
+				return Tab.TabWidget->IsForeground();
+			}
+		}
+	}
+
+	return false;
+}
+
+void SImTabGroup::SetSizeCoefficient(float SizeCoefficient)
+{
+	CurrentSizeCoefficient = SizeCoefficient;
+}
+
+void SImTabGroup::BeginTabSplitter(EOrientation Orientation)
+{
+	TSharedRef<FTabManager::FSplitter> TabSplitter = [this, Orientation]()
+	{
+		checkf(!WorkingLayoutNodeStack.IsEmpty(), TEXT("The Tab Group's WorkingLayoutNodeStack should not be empty. Are your Begin and End function calls mismatched?"));
+		FImTabLayoutNode& ParentNode = WorkingLayoutNodeStack.Last();
+		TSharedPtr<FSlateImTabSplitter> ParentSplitter = StaticCastSharedPtr<FSlateImTabSplitter>(ParentNode.Node->AsSplitter());
+		checkf(ParentSplitter.IsValid(), TEXT("Tab Splitters can only be added to the Tab Group or other Tab Splitters. Are your Begin and End function calls mismatched?"));
+		
+		if (TSharedPtr<FTabManager::FSplitter> ExistingSplitter = ParentSplitter->GetChildAsSplitter(ParentNode.CurrentChildIndex))
+		{
+			if (ExistingSplitter->GetOrientation() != Orientation)
+			{
+				bIsLayoutUpToDate = false;
+				ExistingSplitter->SetOrientation(Orientation);
+			}
+			return ExistingSplitter.ToSharedRef();
+		}
+
+		ParentSplitter->RemoveChildAt(ParentNode.CurrentChildIndex);
+		bIsLayoutUpToDate = false;
+		TSharedRef<FTabManager::FSplitter> TabSplitter = FTabManager::NewSplitter()
+			->SetSizeCoefficient(CurrentSizeCoefficient)
+			->SetOrientation(Orientation);
+		return TabSplitter;
+	}();
+
+	// Reset the size coefficient after using it
+	CurrentSizeCoefficient = 1.f;
+	WorkingLayoutNodeStack.Emplace(TabSplitter);
+}
+
+void SImTabGroup::BeginTabStack()
+{
+	TSharedRef<FTabManager::FStack> TabStack = [this]()
+	{
+		checkf(!WorkingLayoutNodeStack.IsEmpty(), TEXT("The Tab Group's WorkingLayoutNodeStack should not be empty. Are your Begin and End function calls mismatched?"));
+		FImTabLayoutNode& ParentNode = WorkingLayoutNodeStack.Last();
+		TSharedPtr<FSlateImTabSplitter> ParentSplitter = StaticCastSharedPtr<FSlateImTabSplitter>(ParentNode.Node->AsSplitter());
+		checkf(ParentSplitter.IsValid(), TEXT("Tab Stacks can only be added to the Tab Group or other Tab Splitters. Are your Begin and End function calls mismatched?"));
+		
+		if (TSharedPtr<FTabManager::FStack> ExistingStack = ParentSplitter->GetChildAsStack(ParentNode.CurrentChildIndex))
+		{
+			return ExistingStack.ToSharedRef();
+		}
+
+		ParentSplitter->RemoveChildAt(ParentNode.CurrentChildIndex);
+		bIsLayoutUpToDate = false;
+		TSharedRef<FTabManager::FStack> TabStack = FTabManager::NewStack()
+			->SetSizeCoefficient(CurrentSizeCoefficient)
+			->SetHideTabWell(false);
+		return TabStack;
+	}();
+	
+	// Reset the size coefficient after using it
+	CurrentSizeCoefficient = 1.f;
+	WorkingLayoutNodeStack.Emplace(TabStack);
+}
+
+void SImTabGroup::EndTabLayoutNode()
+{
+	if (ensureMsgf(!WorkingLayoutNodeStack.IsEmpty(), TEXT("The Tab Group's WorkingLayoutNodeStack should not be empty. Are your Begin and End function calls mismatched?")))
+	{
+		const FImTabLayoutNode& Node = WorkingLayoutNodeStack.Pop(EAllowShrinking::No);
+		if (ensureMsgf(!WorkingLayoutNodeStack.IsEmpty(), TEXT("The Tab Group's LayoutNodeStack should not be empty. Are your Begin and End function calls mismatched?")))
+		{
+			// Remove unused children in the Node
+			if (const TSharedPtr<FSlateImTabSplitter> NodeSplitter = StaticCastSharedPtr<FSlateImTabSplitter>(Node.Node->AsSplitter()))
+			{
+				if (NodeSplitter->RemoveUnusedChildren(Node.CurrentChildIndex))
+				{
+					bIsLayoutUpToDate = false;
+				}
+			}
+			
+			FImTabLayoutNode& ParentNode = WorkingLayoutNodeStack.Last();
+			if (const TSharedPtr<FSlateImTabSplitter> Splitter = StaticCastSharedPtr<FSlateImTabSplitter>(ParentNode.Node->AsSplitter()))
+			{
+				if (Splitter->SplitAt(ParentNode.CurrentChildIndex, Node.Node))
+				{
+					bIsLayoutUpToDate = false;
+				}
+				++ParentNode.CurrentChildIndex;
+			}
+			else if (ParentNode.Node->AsStack())
+			{
+				ensureMsgf(false, TEXT("Tab Stacks can only contain tabs. Are your Begin and End function calls mismatched?"));
+			}
+			else
+			{
+				ensureMsgf(false, TEXT("Unhandled tab Layout Node type. Are your Begin and End function calls mismatched?"));
+			}
+		}
+	}
+}
+
+TSharedRef<SDockTab> SImTabGroup::SpawnTab(const FSpawnTabArgs& Args)
+{
+	if (TSharedPtr<SDockTab> ExistingTab = ReuseTab(Args.GetTabId()))
+	{
+		return ExistingTab.ToSharedRef();
+	}
+	
+	return SNew(SDockTab)
+	[
+		SNew(STextBlock)
+		.Text(FText::Format(NSLOCTEXT("SlateIM", "InvalidTabIdFormat", "Attempting to display invalid Tab Id: {0}"), FText::FromString(Args.GetTabId().ToString())))
+	];
+}
+
+TSharedPtr<SDockTab> SImTabGroup::ReuseTab(const FTabId& TabId) const
+{
+	const FImTab* TabPtr = Tabs.FindByPredicate([TabId](const FImTab& Tab)
+	{
+		return Tab.TabId == TabId.TabType;
+	});
+
+	if (TabPtr != nullptr && TabPtr->TabWidget.IsValid())
+	{
+		return TabPtr->TabWidget.ToSharedRef();
+	}
+
+	return nullptr;
+}
diff --git a/Engine/Plugins/Experimental/SlateIM/Source/SlateIM/Private/Containers/SImTabGroup.h b/Engine/Plugins/Experimental/SlateIM/Source/SlateIM/Private/Containers/SImTabGroup.h
new file mode 100644
--- /dev/null
+++ b/Engine/Plugins/Experimental/SlateIM/Source/SlateIM/Private/Containers/SImTabGroup.h
@@ -0,0 +1,68 @@
+﻿// Copyright Epic Games, Inc. All Rights Reserved.
+
+#pragma once
+
+#include "Framework/Docking/TabManager.h"
+#include "ISlateIMContainer.h"
+#include "Widgets/SCompoundWidget.h"
+
+class SImTab;
+
+class SLATEIM_API SImTabGroup : public SCompoundWidget, public ISlateIMContainer
+{
+	SLATE_DECLARE_WIDGET(SImTabGroup, SCompoundWidget)
+	SLATE_IM_TYPE_DATA(SImTabGroup, ISlateIMContainer)
+	
+public:
+	SLATE_BEGIN_ARGS(SImTabGroup)
+		{}
+		SLATE_ARGUMENT(FName, TabGroupId)
+	SLATE_END_ARGS()
+
+	void Construct(const FArguments& InArgs);
+	virtual ~SImTabGroup() override;
+
+	void BeginTabGroupUpdates();
+	void FinalizeTabGroup();
+
+	virtual int32 GetNumChildren() override;
+	virtual FSlateIMChild GetChild(int32 Index) override;
+	virtual void UpdateChild(FSlateIMChild Child, int32 Index, const FSlateIMSlotData& AlignmentData) override;
+	virtual void RemoveUnusedChildren(int32 LastUsedChildIndex) override;
+
+	virtual FSlateIMChild GetContainer() override
+	{
+		return AsShared();
+	}
+
+	FName GetTabGroupId() const;
+	bool IsTabForegrounded(const FName& TabId) const;
+
+	void SetSizeCoefficient(float SizeCoefficient);
+	
+	void BeginTabSplitter(EOrientation Orientation);
+	void BeginTabStack();
+	void EndTabLayoutNode();
+
+	TSharedRef<SDockTab> SpawnTab(const FSpawnTabArgs& Args);
+	TSharedPtr<SDockTab> ReuseTab(const FTabId& TabId) const;
+
+private:
+	struct FImTab
+	{
+		FName TabId = NAME_None;
+		TSharedPtr<SImTab> TabWidget;
+	};
+	struct FImTabLayoutNode
+	{
+		TSharedRef<FTabManager::FLayoutNode> Node;
+		int32 CurrentChildIndex = 0;
+	};
+	TSharedPtr<FTabManager> TabManager;
+	TSharedPtr<FTabManager::FLayout> TabLayout;
+	TArray<FImTabLayoutNode> WorkingLayoutNodeStack;
+	TArray<FImTab> Tabs;
+	FName TabGroupId = NAME_None;
+	float CurrentSizeCoefficient = 1.f;
+	bool bIsLayoutUpToDate = false;
+};
diff --git a/Engine/Plugins/Experimental/SlateIM/Source/SlateIM/Private/Implementation/SlateIM_Tabs.cpp b/Engine/Plugins/Experimental/SlateIM/Source/SlateIM/Private/Implementation/SlateIM_Tabs.cpp
new file mode 100644
--- /dev/null
+++ b/Engine/Plugins/Experimental/SlateIM/Source/SlateIM/Private/Implementation/SlateIM_Tabs.cpp
@@ -0,0 +1,126 @@
+﻿// Copyright Epic Games, Inc. All Rights Reserved.
+
+#include "SlateIM.h"
+
+#include "Containers/SImTab.h"
+#include "Containers/SImTabGroup.h"
+#include "Misc/SlateIMLogging.h"
+#include "Misc/SlateIMManager.h"
+#include "Misc/SlateIMWidgetScope.h"
+
+namespace SlateIM
+{	
+	void BeginTabGroup(const FName& TabGroupId)
+	{
+		TSharedPtr<SImTabGroup> ContainerWidget;
+
+		// Note widget scope must fall out of scope before we push container or else we are pushing this new container into itself instead of into its parent
+		{
+			FWidgetScope<SImTabGroup> Scope;
+			ContainerWidget = Scope.GetWidget();
+
+			if (!ContainerWidget || ContainerWidget->GetTabGroupId() != TabGroupId)
+			{
+				ContainerWidget = SNew(SImTabGroup)
+					.TabGroupId(TabGroupId);
+				Scope.UpdateWidget(ContainerWidget);
+			}
+
+			ContainerWidget->BeginTabGroupUpdates();
+		}
+
+		FSlateIMManager::Get().PushContainer(FContainerNode(ContainerWidget));
+	}
+	
+	void EndTabGroup()
+	{
+		if (TSharedPtr<SImTabGroup> TabGroup = FSlateIMManager::Get().GetCurrentContainer<SImTabGroup>())
+		{
+			TabGroup->FinalizeTabGroup();
+		}
+		FSlateIMManager::Get().PopContainer<SImTabGroup>();
+	}
+
+	void BeginTabStack()
+	{
+		TSharedPtr<SImTabGroup> TabGroup = FSlateIMManager::Get().GetCurrentContainer<SImTabGroup>();
+		if (ensureMsgf(TabGroup, TEXT("Cannot begin a Tab Stack - the current container is not a Tab Group. Are your Begin and End function calls mismatched?")))
+		{
+			TabGroup->BeginTabStack();
+		}
+	}
+
+	void EndTabStack()
+	{
+		TSharedPtr<SImTabGroup> TabGroup = FSlateIMManager::Get().GetCurrentContainer<SImTabGroup>();
+		if (ensureMsgf(TabGroup, TEXT("Cannot end a Tab Stack - the current container is not a Tab Group. Are your Begin and End function calls mismatched?")))
+		{
+			TabGroup->EndTabLayoutNode();
+		}
+	}
+
+	void BeginTabSplitter(EOrientation Orientation)
+	{
+		TSharedPtr<SImTabGroup> TabGroup = FSlateIMManager::Get().GetCurrentContainer<SImTabGroup>();
+		if (ensureMsgf(TabGroup, TEXT("Cannot begin a Tab Splitter - the current container is not a Tab Group. Are your Begin and End function calls mismatched?")))
+		{
+			TabGroup->BeginTabSplitter(Orientation);
+		}
+	}
+
+	void EndTabSplitter()
+	{
+		TSharedPtr<SImTabGroup> TabGroup = FSlateIMManager::Get().GetCurrentContainer<SImTabGroup>();
+		if (ensureMsgf(TabGroup, TEXT("Cannot end a Tab Splitter - the current container is not a Tab Group. Are your Begin and End function calls mismatched?")))
+		{
+			TabGroup->EndTabLayoutNode();
+		}
+	}
+
+	void TabSplitterSizeCoefficient(float SizeCoefficient)
+	{
+		TSharedPtr<SImTabGroup> TabGroup = FSlateIMManager::Get().GetCurrentContainer<SImTabGroup>();
+		if (ensureMsgf(TabGroup, TEXT("Cannot set the Tab Splitter Size Coefficient - the current container is not a Tab Group. Are your Begin and End function calls mismatched?")))
+		{
+			TabGroup->SetSizeCoefficient(SizeCoefficient);
+		}
+	}
+
+	bool BeginTab(const FName& TabId, const FSlateIcon& TabIcon, const FText& TabTitle)
+	{
+		TSharedPtr<SImTabGroup> TabGroup = FSlateIMManager::Get().GetCurrentContainer<SImTabGroup>();
+		if (!ensureMsgf(TabGroup, TEXT("Cannot begin a Tab - the current container is not a Tab Group. Are your Begin and End function calls mismatched?")))
+		{
+			return false;
+		}
+		
+		TSharedPtr<SImTab> ContainerWidget;
+
+		// Note widget scope must fall out of scope before we push container or else we are pushing this new container into itself instead of into its parent
+		{
+			FWidgetScope<SImTab> Scope;
+			ContainerWidget = Scope.GetWidget();
+
+			if (!ContainerWidget || ContainerWidget->TabId != TabId)
+			{
+				ContainerWidget = SNew(SImTab)
+					.CanEverClose(false);
+				ContainerWidget->TabId = TabId;
+				Scope.UpdateWidget(ContainerWidget);
+
+				// TabIcon and TabTitle are not updated after construction, this could be supported in the future but is not currently part of the Tab Manager API
+				ContainerWidget->TabIcon = TabIcon;
+				ContainerWidget->TabTitle = TabTitle;
+			}
+		}
+
+		FSlateIMManager::Get().PushContainer(FContainerNode(ContainerWidget));
+
+		return TabGroup->IsTabForegrounded(TabId);
+	}
+
+	void EndTab()
+	{
+		FSlateIMManager::Get().PopContainer<SImTab>();
+	}
+}
diff --git a/Engine/Plugins/Experimental/SlateIM/Source/SlateIM/Private/Implementation/SlateIM_Widgets.cpp b/Engine/Plugins/Experimental/SlateIM/Source/SlateIM/Private/Implementation/SlateIM_Widgets.cpp
--- a/Engine/Plugins/Experimental/SlateIM/Source/SlateIM/Private/Implementation/SlateIM_Widgets.cpp
+++ b/Engine/Plugins/Experimental/SlateIM/Source/SlateIM/Private/Implementation/SlateIM_Widgets.cpp
@@ -59,6 +59,7 @@ namespace SlateIM
 		{
 			TextBlock =
 				SNew(STextBlock)
+				.AutoWrapText(true)
 				.Text(FText::FromStringView(InText))
 				.ColorAndOpacity(Color);
 
diff --git a/Engine/Plugins/Experimental/SlateIM/Source/SlateIM/Private/Misc/SlateIMExamples.cpp b/Engine/Plugins/Experimental/SlateIM/Source/SlateIM/Private/Misc/SlateIMExamples.cpp
--- a/Engine/Plugins/Experimental/SlateIM/Source/SlateIM/Private/Misc/SlateIMExamples.cpp
+++ b/Engine/Plugins/Experimental/SlateIM/Source/SlateIM/Private/Misc/SlateIMExamples.cpp
@@ -3,6 +3,7 @@
 #if WITH_SLATEIM_EXAMPLES
 #include "SlateIMExamples.h"
 
+#include "Framework/Application/SlateApplication.h"
 #include "HAL/PlatformApplicationMisc.h"
 #include "Misc/App.h"
 #include "ProfilingDebugging/ScopedTimers.h"
@@ -11,6 +12,8 @@
 #include "Styling/AppStyle.h"
 #include "Styling/SlateStyle.h"
 #include "Styling/StyleColors.h"
+#include "Widgets/Docking/SDockTab.h"
+#include "Widgets/Text/STextBlock.h"
 
 #if WITH_EDITOR
 #include "SLevelViewport.h"
@@ -118,7 +121,48 @@ void FSlateIMTestWidget::Draw()
 	}
 	SlateIM::EndHorizontalStack();
 	
-	// Full root scrollbox
+	SlateIM::Fill();
+	SlateIM::HAlign(HAlign_Fill);
+	SlateIM::VAlign(VAlign_Fill);
+	SlateIM::BeginTabGroup(TEXT("ExampleContent"));
+	SlateIM::BeginTabStack();
+	if (SlateIM::BeginTab(TEXT("Basics")))
+	{
+		DrawBasics();
+	}
+	SlateIM::EndTab();
+
+	if (SlateIM::BeginTab(TEXT("Lists and Tables")))
+	{
+		DrawTables();
+	}
+	SlateIM::EndTab();
+
+	if (SlateIM::BeginTab(TEXT("Graphs")))
+	{
+		DrawGraphs();
+	}
+	SlateIM::EndTab();
+
+	if (SlateIM::BeginTab(TEXT("Inputs")))
+	{
+		DrawInputs();
+	}
+	SlateIM::EndTab();
+
+	if (SlateIM::BeginTab(TEXT("Tabs")))
+	{
+		DrawTabs();
+	}
+	SlateIM::EndTab();
+	SlateIM::EndTabStack();
+	SlateIM::EndTabGroup();
+	SlateIM::EndBorder();
+}
+
+void FSlateIMTestWidget::DrawBasics()
+{
+	SCOPED_NAMED_EVENT_TEXT("Basics", FColorList::Goldenrod);
 	SlateIM::Fill();
 	SlateIM::HAlign(HAlign_Fill);
 	SlateIM::VAlign(VAlign_Fill);
@@ -183,7 +227,7 @@ void FSlateIMTestWidget::Draw()
 				SlateIM::Text(TEXT("Text With Color"), FLinearColor::Green);
 				SlateIM::Text(TEXT("Text With style color"), FStyleColors::Primary);
 			}
-		
+	
 			if (SlateIM::CheckBox(TEXT("Undetermined Check Box"), CheckStateEnum))
 			{
 				UE_LOG(LogSlateIM, Log, TEXT("Check Box State Changed"));
@@ -211,7 +255,7 @@ void FSlateIMTestWidget::Draw()
 				: RedIcon.LoadSynchronous();
 			SlateIM::Image(Texture);
 #endif
-		
+	
 			SlateIM::BeginVerticalStack();
 			SlateIM::Button(TEXT("Button 1"));
 			SlateIM::Button(TEXT("Button 2"));
@@ -293,7 +337,7 @@ void FSlateIMTestWidget::Draw()
 			SlateIM::Image("AppIcon");
 			SlateIM::EndHorizontalStack();
 		}
-		
+	
 		{
 			SCOPED_NAMED_EVENT_TEXT("PopUp example", FColorList::Goldenrod);
 			// PopUp example
@@ -309,7 +353,7 @@ void FSlateIMTestWidget::Draw()
 			}
 			SlateIM::EndHorizontalStack();
 		}
-		
+	
 		{
 			SCOPED_NAMED_EVENT_TEXT("DisabledState example", FColorList::Goldenrod);
 			// DisabledState example
@@ -354,7 +398,7 @@ void FSlateIMTestWidget::Draw()
 			SlateIM::EndSubMenu();
 			SlateIM::EndContextMenuAnchor();
 		}
-	
+
 
 		// Modal Examples
 		{
@@ -431,7 +475,7 @@ void FSlateIMTestWidget::Draw()
 					SlateIM::Text(TEXT("UNHANDLED RESULT"));
 					break;
 				}
-			
+		
 				if (SlateIM::Button(TEXT("Reset")))
 				{
 					DialogResult.Reset();
@@ -442,27 +486,45 @@ void FSlateIMTestWidget::Draw()
 				SlateIM::Text(TEXT("No Dialog Result"));
 			}
 		}
-		
-		SlateIM::Text(TEXT("Scroll Bar Test"));
-		SlateIM::BeginHorizontalStack();
 
-		SlateIM::Text(TEXT("Num Items "));
-		SlateIM::MinWidth(50.f);
-		if (SlateIM::EditableText(NumItemsText))
+		if (bShouldBeDisabled)
 		{
-			LiveNumItems = FCString::Atoi(*NumItemsText);
+			SlateIM::EndDisabledState();
 		}
+	}
+	SlateIM::EndScrollBox();
 
-		if (SlateIM::Button(TEXT("Regenerate")))
-		{
-			NumItems = FCString::Atoi(*NumItemsText);
-		}
+	bRefreshComboItems = false;
+}
 
-		int32 NumTableItems = NumItems;
-		SlateIM::CheckBox(TEXT("Live Update Table?"), bShouldLiveUpdateTable);
-		if (bShouldLiveUpdateTable)
+void FSlateIMTestWidget::DrawTables()
+{
+	SCOPED_NAMED_EVENT_TEXT("Lists and Tables", FColorList::Goldenrod);
+	SlateIM::HAlign(HAlign_Fill);
+	SlateIM::VAlign(VAlign_Fill);
+	SlateIM::BeginScrollBox();
+	{
+		int32 NumTableItems = 0;
+		SlateIM::BeginHorizontalStack();
 		{
-			NumTableItems = LiveNumItems;
+			SlateIM::Text(TEXT("Num Items "));
+			SlateIM::MinWidth(50.f);
+			if (SlateIM::EditableText(NumItemsText))
+			{
+				LiveNumItems = FCString::Atoi(*NumItemsText);
+			}
+
+			if (SlateIM::Button(TEXT("Regenerate")))
+			{
+				NumItems = FCString::Atoi(*NumItemsText);
+			}
+
+			NumTableItems = NumItems;
+			SlateIM::CheckBox(TEXT("Live Update Table?"), bShouldLiveUpdateTable);
+			if (bShouldLiveUpdateTable)
+			{
+				NumTableItems = LiveNumItems;
+			}
 		}
 		SlateIM::EndHorizontalStack();
 
@@ -533,75 +595,7 @@ void FSlateIMTestWidget::Draw()
 			SlateIM::EndTable();
 		}
 		SlateIM::EndHorizontalStack();
-
-		// Graph Examples
-		{
-			SCOPED_NAMED_EVENT_TEXT("Graph Example", FColorList::Goldenrod);
-			SlateIM::BeginHorizontalStack();
-			{
-				if (SquareGraphValues.Num() >= 100)
-				{
-					SquareGraphValues.PopFront();
-				}
-				const double SquareValue = (GFrameCounter / 10) % 2;
-				SquareGraphValues.Emplace(SquareValue);
-				SlateIM::Fill();
-				SlateIM::MinHeight(200.f);
-				SlateIM::BeginGraph();
-				SlateIM::GraphLine(SquareGraphValues.Compact(), FLinearColor::White, 3.f, FDoubleRange(0, 1.0));
-				SlateIM::EndGraph();
-
-				SlateIM::VAlign(VAlign_Center);
-				SlateIM::Text(FString::Printf(TEXT("%0.3lf"), SquareValue));
-
-				const double NextSinX = SinGraphValues.Last().X + 1.0;
-				const double NextCosX = CosGraphValues.Last().X + 1.0;
-				const double NextTanX = TanGraphValues.Last().X + 1.0;
-				if (GFrameCounter % 4 == 0)
-				{
-					if (SinGraphValues.Num() >= 100)
-					{
-						SinGraphValues.PopFront();
-					}
-					SinGraphValues.Emplace(NextSinX, FMath::Sin(NextSinX/4.0));
-					
-					if (CosGraphValues.Num() >= 100)
-					{
-						CosGraphValues.PopFront();
-					}
-					CosGraphValues.Emplace(NextCosX, FMath::Cos(NextCosX/4.0));
-					
-					if (TanGraphValues.Num() >= 100)
-					{
-						TanGraphValues.PopFront();
-					}
-					TanGraphValues.Emplace(NextTanX, FMath::Tan(NextTanX/4.0));
-				}
-				SlateIM::Fill();
-				SlateIM::MinHeight(200.f);
-				SlateIM::BeginGraph();
-				SlateIM::GraphLine(SinGraphValues.Compact(), FColor::Orange, 1.f, FDoubleRange(NextSinX - 100.0, NextSinX), FDoubleRange(-1.5, 1.5));
-				SlateIM::GraphLine(CosGraphValues.Compact(), FLinearColor::Green, 1.f, FDoubleRange(NextCosX - 100.0, NextCosX), FDoubleRange(-1.5, 1.5));
-				SlateIM::GraphLine(TanGraphValues.Compact(), FColor::Magenta, 1.f, FDoubleRange(NextTanX - 100.0, NextTanX), FDoubleRange(-1.5, 1.5));
-				SlateIM::EndGraph();
-
-				SlateIM::MinWidth(50.f);
-				SlateIM::MaxWidth(50.f);
-				SlateIM::BeginVerticalStack();
-				SlateIM::Fill();
-				SlateIM::VAlign(VAlign_Center);
-				SlateIM::Text(FString::Printf(TEXT("%0.3lf"), SinGraphValues.Last().Y), FColor::Orange);
-				SlateIM::Fill();
-				SlateIM::VAlign(VAlign_Center);
-				SlateIM::Text(FString::Printf(TEXT("%0.3lf"), CosGraphValues.Last().Y), FLinearColor::Green);
-				SlateIM::Fill();
-				SlateIM::VAlign(VAlign_Center);
-				SlateIM::Text(FString::Printf(TEXT("%0.3lf"), TanGraphValues.Last().Y), FColor::Magenta);
-				SlateIM::EndVerticalStack();
-			}
-			SlateIM::EndHorizontalStack();
-		}
-
+		
 		// Tree Example
 		{
 			SCOPED_NAMED_EVENT_TEXT("Tree Example", FColorList::Goldenrod);
@@ -609,6 +603,7 @@ void FSlateIMTestWidget::Draw()
 			SlateIM::MinWidth(500.f);
 			SlateIM::MinHeight(200.f);
 			SlateIM::MaxHeight(200.f);
+			SlateIM::HAlign(HAlign_Fill);
 			SlateIM::VAlign(VAlign_Fill);
 			SlateIM::BeginTable(nullptr, TableRowStyle);
 			SlateIM::AddTableColumn(TEXT("Name"));
@@ -624,7 +619,7 @@ void FSlateIMTestWidget::Draw()
 					SlateIM::VAlign(VAlign_Center);
 					SlateIM::Text(TEXT("Continent"));
 				}
-				
+			
 				if (SlateIM::NextTableCell())
 				{
 					SlateIM::VAlign(VAlign_Center);
@@ -673,7 +668,7 @@ void FSlateIMTestWidget::Draw()
 							}
 						}
 						SlateIM::EndTableRowChildren();
-					
+				
 						if (SlateIM::NextTableCell())
 						{
 							SlateIM::VAlign(VAlign_Center);
@@ -737,7 +732,7 @@ void FSlateIMTestWidget::Draw()
 							}
 						}
 						SlateIM::EndTableRowChildren();
-						
+					
 						if (SlateIM::NextTableCell())
 						{
 							SlateIM::VAlign(VAlign_Center);
@@ -767,100 +762,313 @@ void FSlateIMTestWidget::Draw()
 				}
 				SlateIM::EndTableRowChildren();
 			}
-			
+		
 			SlateIM::EndTable();
 		}
+	}
+	SlateIM::EndScrollBox();
+}
 
-		// Inputs example
+void FSlateIMTestWidget::DrawGraphs()
+{
+	SCOPED_NAMED_EVENT_TEXT("Graphs", FColorList::Goldenrod);
+	SlateIM::BeginHorizontalStack();
+	{
+		if (SquareGraphValues.Num() >= 100)
 		{
-			SlateIM::BeginHorizontalStack();
+			SquareGraphValues.PopFront();
+		}
+		const double SquareValue = (GFrameCounter / 10) % 2;
+		SquareGraphValues.Emplace(SquareValue);
+		SlateIM::Fill();
+		SlateIM::MinHeight(200.f);
+		SlateIM::BeginGraph();
+		SlateIM::GraphLine(SquareGraphValues.Compact(), FLinearColor::White, 3.f, FDoubleRange(0, 1.0));
+		SlateIM::EndGraph();
+
+		SlateIM::VAlign(VAlign_Center);
+		SlateIM::Text(FString::Printf(TEXT("%0.3lf"), SquareValue));
+
+		const double NextSinX = SinGraphValues.Last().X + 1.0;
+		const double NextCosX = CosGraphValues.Last().X + 1.0;
+		const double NextTanX = TanGraphValues.Last().X + 1.0;
+		if (GFrameCounter % 4 == 0)
+		{
+			if (SinGraphValues.Num() >= 100)
 			{
-				SlateIM::BeginVerticalStack();
-				SlateIM::BeginHorizontalStack();
-				SlateIM::Spacer(FVector2D(50, 50));
-				{
-					WBrush = FSlateColorBrush(SlateIM::Private::GetKeyStateColor(EKeys::W));
-					SlateIM::MinWidth(50.f);
-					SlateIM::MinHeight(50.f);
-					SlateIM::MaxWidth(50.f);
-					SlateIM::MaxHeight(50.f);
-					SlateIM::BeginBorder(&WBrush);
-					SlateIM::HAlign(HAlign_Center);
-					SlateIM::VAlign(VAlign_Center);
-					SlateIM::Text(TEXT("W"), FLinearColor::White);
-					SlateIM::EndBorder();
-				}
-				SlateIM::Spacer(FVector2D(50, 50));
-				SlateIM::EndHorizontalStack();
-				SlateIM::BeginHorizontalStack();
+				SinGraphValues.PopFront();
+			}
+			SinGraphValues.Emplace(NextSinX, FMath::Sin(NextSinX/4.0));
+		
+			if (CosGraphValues.Num() >= 100)
+			{
+				CosGraphValues.PopFront();
+			}
+			CosGraphValues.Emplace(NextCosX, FMath::Cos(NextCosX/4.0));
+		
+			if (TanGraphValues.Num() >= 100)
+			{
+				TanGraphValues.PopFront();
+			}
+			TanGraphValues.Emplace(NextTanX, FMath::Tan(NextTanX/4.0));
+		}
+		SlateIM::Fill();
+		SlateIM::MinHeight(200.f);
+		SlateIM::BeginGraph();
+		SlateIM::GraphLine(SinGraphValues.Compact(), FColor::Orange, 1.f, FDoubleRange(NextSinX - 100.0, NextSinX), FDoubleRange(-1.5, 1.5));
+		SlateIM::GraphLine(CosGraphValues.Compact(), FLinearColor::Green, 1.f, FDoubleRange(NextCosX - 100.0, NextCosX), FDoubleRange(-1.5, 1.5));
+		SlateIM::GraphLine(TanGraphValues.Compact(), FColor::Magenta, 1.f, FDoubleRange(NextTanX - 100.0, NextTanX), FDoubleRange(-1.5, 1.5));
+		SlateIM::EndGraph();
+
+		SlateIM::MinWidth(50.f);
+		SlateIM::MaxWidth(50.f);
+		SlateIM::BeginVerticalStack();
+		SlateIM::Fill();
+		SlateIM::VAlign(VAlign_Center);
+		SlateIM::Text(FString::Printf(TEXT("%0.3lf"), SinGraphValues.Last().Y), FColor::Orange);
+		SlateIM::Fill();
+		SlateIM::VAlign(VAlign_Center);
+		SlateIM::Text(FString::Printf(TEXT("%0.3lf"), CosGraphValues.Last().Y), FLinearColor::Green);
+		SlateIM::Fill();
+		SlateIM::VAlign(VAlign_Center);
+		SlateIM::Text(FString::Printf(TEXT("%0.3lf"), TanGraphValues.Last().Y), FColor::Magenta);
+		SlateIM::EndVerticalStack();
+	}
+	SlateIM::EndHorizontalStack();
+}
+
+void FSlateIMTestWidget::DrawInputs()
+{
+	SCOPED_NAMED_EVENT_TEXT("Inputs", FColorList::Goldenrod);
+	SlateIM::HAlign(HAlign_Center);
+	SlateIM::BeginHorizontalStack();
+	{
+		SlateIM::BeginVerticalStack();
+		SlateIM::BeginHorizontalStack();
+		SlateIM::Spacer(FVector2D(50, 50));
+		{
+			WBrush = FSlateColorBrush(SlateIM::Private::GetKeyStateColor(EKeys::W));
+			SlateIM::MinWidth(50.f);
+			SlateIM::MinHeight(50.f);
+			SlateIM::MaxWidth(50.f);
+			SlateIM::MaxHeight(50.f);
+			SlateIM::BeginBorder(&WBrush);
+			SlateIM::HAlign(HAlign_Center);
+			SlateIM::VAlign(VAlign_Center);
+			SlateIM::Text(TEXT("W"), FLinearColor::White);
+			SlateIM::EndBorder();
+		}
+		SlateIM::Spacer(FVector2D(50, 50));
+		SlateIM::EndHorizontalStack();
+		SlateIM::BeginHorizontalStack();
+		{
+			ABrush = FSlateColorBrush(SlateIM::Private::GetKeyStateColor(EKeys::A));
+			SlateIM::MinWidth(50.f);
+			SlateIM::MinHeight(50.f);
+			SlateIM::MaxWidth(50.f);
+			SlateIM::MaxHeight(50.f);
+			SlateIM::BeginBorder(&ABrush);
+			SlateIM::HAlign(HAlign_Center);
+			SlateIM::VAlign(VAlign_Center);
+			SlateIM::Text(TEXT("A"), FLinearColor::White);
+			SlateIM::EndBorder();
+		}
+		{
+			SBrush = FSlateColorBrush(SlateIM::Private::GetKeyStateColor(EKeys::S));
+			SlateIM::MinWidth(50.f);
+			SlateIM::MinHeight(50.f);
+			SlateIM::MaxWidth(50.f);
+			SlateIM::MaxHeight(50.f);
+			SlateIM::BeginBorder(&SBrush);
+			SlateIM::HAlign(HAlign_Center);
+			SlateIM::VAlign(VAlign_Center);
+			SlateIM::Text(TEXT("S"), FLinearColor::White);
+			SlateIM::EndBorder();
+		}
+		{
+			DBrush = FSlateColorBrush(SlateIM::Private::GetKeyStateColor(EKeys::D));
+			SlateIM::MinWidth(50.f);
+			SlateIM::MinHeight(50.f);
+			SlateIM::MaxWidth(50.f);
+			SlateIM::MaxHeight(50.f);
+			SlateIM::BeginBorder(&DBrush);
+			SlateIM::HAlign(HAlign_Center);
+			SlateIM::VAlign(VAlign_Center);
+			SlateIM::Text(TEXT("D"), FLinearColor::White);
+			SlateIM::EndBorder();
+		}
+		SlateIM::EndHorizontalStack();
+		SlateIM::EndVerticalStack();
+	}
+
+	{
+		SlateIM::BeginVerticalStack();
+		SlateIM::Text(TEXT("Right Stick X-value"));
+		const float NormalizedAnalogXValue = (1.f + SlateIM::GetKeyAnalogValue(EKeys::Gamepad_RightX)) * 0.5f;
+		SlateIM::VAlign(VAlign_Center);
+		SlateIM::ProgressBar(NormalizedAnalogXValue);
+	
+		SlateIM::Fill();
+		SlateIM::Spacer(FVector2D(1.f, 1.f));
+	
+		SlateIM::Text(TEXT("Right Stick Y-value"));
+		const float NormalizedAnalogYValue = (1.f + SlateIM::GetKeyAnalogValue(EKeys::Gamepad_RightY)) * 0.5f;
+		SlateIM::VAlign(VAlign_Center);
+		SlateIM::ProgressBar(NormalizedAnalogYValue);
+		SlateIM::EndVerticalStack();
+	}
+	SlateIM::EndHorizontalStack();
+}
+
+void FSlateIMTestWidget::DrawTabs()
+{
+	SCOPED_NAMED_EVENT_TEXT("Tabs", FColorList::Goldenrod);
+	{
+		SlateIM::Text(TEXT("Splitters"));
+		SlateIM::HAlign(HAlign_Fill);
+		SlateIM::VAlign(VAlign_Fill);
+		SlateIM::Fill();
+		SlateIM::BeginTabGroup(TEXT("SlateIMTabTesting1"));
+		{
+			SlateIM::BeginTabSplitter(Orient_Horizontal);
+			{
+				SlateIM::TabSplitterSizeCoefficient(0.3f);
+				SlateIM::BeginTabStack();
 				{
-					ABrush = FSlateColorBrush(SlateIM::Private::GetKeyStateColor(EKeys::A));
-					SlateIM::MinWidth(50.f);
-					SlateIM::MinHeight(50.f);
-					SlateIM::MaxWidth(50.f);
-					SlateIM::MaxHeight(50.f);
-					SlateIM::BeginBorder(&ABrush);
-					SlateIM::HAlign(HAlign_Center);
-					SlateIM::VAlign(VAlign_Center);
-					SlateIM::Text(TEXT("A"), FLinearColor::White);
-					SlateIM::EndBorder();
+					if (SlateIM::BeginTab(TEXT("Left Tab"), FSlateIcon(FAppStyle::GetAppStyleSetName(), TEXT("Symbols.LeftArrow"))))
+					{
+						// Fill so that the text autowrapping updates properly as the splitter is adjusted by the user
+						SlateIM::HAlign(HAlign_Fill);
+						SlateIM::Text(TEXT("Left Tab taking up 30% of the horizontal space"));
+					}
+					SlateIM::EndTab();
+					if (SlateIM::BeginTab(TEXT("Right Tab"), FSlateIcon(FAppStyle::GetAppStyleSetName(), TEXT("Symbols.RightArrow"))))
+					{
+						// Fill so that the text autowrapping updates properly as the splitter is adjusted by the user
+						SlateIM::HAlign(HAlign_Fill);
+						SlateIM::Text(TEXT("Right Tab taking up 30% of the horizontal space"));
+					}
+					SlateIM::EndTab();
 				}
+				SlateIM::EndTabStack();
+			
+				SlateIM::TabSplitterSizeCoefficient(0.7f);
+				SlateIM::BeginTabSplitter(Orient_Vertical);
 				{
-					SBrush = FSlateColorBrush(SlateIM::Private::GetKeyStateColor(EKeys::S));
-					SlateIM::MinWidth(50.f);
-					SlateIM::MinHeight(50.f);
-					SlateIM::MaxWidth(50.f);
-					SlateIM::MaxHeight(50.f);
-					SlateIM::BeginBorder(&SBrush);
-					SlateIM::HAlign(HAlign_Center);
-					SlateIM::VAlign(VAlign_Center);
-					SlateIM::Text(TEXT("S"), FLinearColor::White);
-					SlateIM::EndBorder();
+					SlateIM::TabSplitterSizeCoefficient(0.3f);
+					SlateIM::BeginTabStack();
+					{
+						if (SlateIM::BeginTab(TEXT("Tab 1")))
+						{
+							// Fill so that the text autowrapping updates properly as the splitter is adjusted by the user
+							SlateIM::HAlign(HAlign_Fill);
+							SlateIM::Text(TEXT("Tab that takes up 30% of the vertical space"));
+						}
+						SlateIM::EndTab();
+					}
+					SlateIM::EndTabStack();
+				
+					SlateIM::TabSplitterSizeCoefficient(0.7f);
+					SlateIM::BeginTabStack();
+					{
+						if (SlateIM::BeginTab(TEXT("Tab 2")))
+						{
+							// Fill so that the text autowrapping updates properly as the splitter is adjusted by the user
+							SlateIM::HAlign(HAlign_Fill);
+							SlateIM::Text(TEXT("Tab that takes up 70% of the vertical space"));
+						}
+						SlateIM::EndTab();
+					}
+					SlateIM::EndTabStack();
 				}
+				SlateIM::EndTabSplitter();
+			}
+			SlateIM::EndTabSplitter();
+		}
+		SlateIM::EndTabGroup();
+	}
+	SlateIM::Spacer(FVector2D(1.f, 20.f));
+	{
+		SlateIM::Text(TEXT("Dynamic Tabs"));
+		SlateIM::BeginHorizontalStack();
+		if (SlateIM::Button(TEXT("Add Tab")))
+		{
+			++DynamicTabCount;
+		}
+		if (SlateIM::Button(TEXT("Remove Tab")))
+		{
+			DynamicTabCount = FMath::Max(0, DynamicTabCount - 1);
+		}
+		SlateIM::EndHorizontalStack();
+		
+		SlateIM::HAlign(HAlign_Fill);
+		SlateIM::VAlign(VAlign_Fill);
+		SlateIM::Fill();
+		SlateIM::BeginTabGroup(TEXT("SlateIMTabTesting2"));
+		{
+			SlateIM::BeginTabStack();
+			{
+				for (int32 i = 0; i < DynamicTabCount; i++)
 				{
-					DBrush = FSlateColorBrush(SlateIM::Private::GetKeyStateColor(EKeys::D));
-					SlateIM::MinWidth(50.f);
-					SlateIM::MinHeight(50.f);
-					SlateIM::MaxWidth(50.f);
-					SlateIM::MaxHeight(50.f);
-					SlateIM::BeginBorder(&DBrush);
-					SlateIM::HAlign(HAlign_Center);
-					SlateIM::VAlign(VAlign_Center);
-					SlateIM::Text(TEXT("D"), FLinearColor::White);
-					SlateIM::EndBorder();
+					if (SlateIM::BeginTab(*FString::Printf(TEXT("DynamicTab%d"), i), FSlateIcon(), FText::Format(INVTEXT("Tab {0}"), FText::AsNumber(i + 1))))
+					{
+						SlateIM::Text(FString::Printf(TEXT("Tab %d Content"), i + 1));
+					}
+					SlateIM::EndTab();
 				}
-				SlateIM::EndHorizontalStack();
-				SlateIM::EndVerticalStack();
 			}
+			SlateIM::EndTabStack();
+		}
+		SlateIM::EndTabGroup();
+	}
+	SlateIM::Spacer(FVector2D(1.f, 20.f));
+	{
+		SlateIM::Text(TEXT("Nested Tab Groups"));
+		SlateIM::HAlign(HAlign_Fill);
+		SlateIM::VAlign(VAlign_Fill);
+		SlateIM::Fill();
+		SlateIM::BeginTabGroup(TEXT("SlateIMTabTesting3"));
+		SlateIM::BeginTabStack();
+		{
+			if (SlateIM::BeginTab(TEXT("Tab 1")))
+			{
+				SlateIM::Text(TEXT("See Tab 2 for Nested Tabs"));
+			}
+			SlateIM::EndTab();
 
+			if (SlateIM::BeginTab(TEXT("Tab 2")))
 			{
-				SlateIM::BeginVerticalStack();
-				SlateIM::Text(TEXT("Right Stick X-value"));
-				const float NormalizedAnalogXValue = (1.f + SlateIM::GetKeyAnalogValue(EKeys::Gamepad_RightX)) * 0.5f;
-				SlateIM::VAlign(VAlign_Center);
-				SlateIM::ProgressBar(NormalizedAnalogXValue);
-				
 				SlateIM::Fill();
-				SlateIM::Spacer(FVector2D(1.f, 1.f));
-				
-				SlateIM::Text(TEXT("Right Stick Y-value"));
-				const float NormalizedAnalogYValue = (1.f + SlateIM::GetKeyAnalogValue(EKeys::Gamepad_RightY)) * 0.5f;
-				SlateIM::VAlign(VAlign_Center);
-				SlateIM::ProgressBar(NormalizedAnalogYValue);
-				SlateIM::EndVerticalStack();
+				SlateIM::HAlign(HAlign_Fill);
+				SlateIM::BeginTabGroup(TEXT("SlateIMTabTestingNestedTabs"));
+				SlateIM::BeginTabStack();
+				{
+					if (SlateIM::BeginTab(TEXT("Nested Tab 1")))
+					{
+						SlateIM::Text(TEXT("Nested Tab 1 Content"));
+					}
+					SlateIM::EndTab();
+			
+					if (SlateIM::BeginTab(TEXT("Nested Tab 2")))
+					{
+						SlateIM::Text(TEXT("Nested Tab 2 Content"));
+					}
+					SlateIM::EndTab();
+			
+					if (SlateIM::BeginTab(TEXT("Nested Tab 3")))
+					{
+						SlateIM::Text(TEXT("Nested Tab 3 Content"));
+					}
+					SlateIM::EndTab();
+				}
+				SlateIM::EndTabStack();
+				SlateIM::EndTabGroup();
 			}
-			SlateIM::EndHorizontalStack();
-		}
-
-		if (bShouldBeDisabled)
-		{
-			SlateIM::EndDisabledState();
+			SlateIM::EndTab();
 		}
+		SlateIM::EndTabStack();
+		SlateIM::EndTabGroup();
 	}
-	SlateIM::EndScrollBox();
-	SlateIM::EndBorder();
-
-	bRefreshComboItems = false;
 }
 
 void FSlateStyleBrowser::DrawWindow(float DeltaTime)
diff --git a/Engine/Plugins/Experimental/SlateIM/Source/SlateIM/Private/Misc/SlateIMExamples.h b/Engine/Plugins/Experimental/SlateIM/Source/SlateIM/Private/Misc/SlateIMExamples.h
--- a/Engine/Plugins/Experimental/SlateIM/Source/SlateIM/Private/Misc/SlateIMExamples.h
+++ b/Engine/Plugins/Experimental/SlateIM/Source/SlateIM/Private/Misc/SlateIMExamples.h
@@ -14,6 +14,11 @@ class FSlateIMTestWidget
 {
 public:
 	void Draw();
+	void DrawBasics();
+	void DrawTables();
+	void DrawGraphs();
+	void DrawInputs();
+	void DrawTabs();
 	
 private:
 	double CurrentTime = 0;
@@ -38,6 +43,7 @@ private:
 	bool bShouldLiveUpdateTable = false;
 	TArray<FString> ComboBoxItems = { TEXT("Option 1"), TEXT("Option 2"), TEXT("Option 3"), TEXT("Option 4") };
 	bool bRefreshComboItems = false;
+	int32 DynamicTabCount = 3;
 	TRingBuffer<double> SquareGraphValues;
 	TRingBuffer<FVector2D> SinGraphValues = { FVector2D(0, FMath::Sin(0.0)) };
 	TRingBuffer<FVector2D> CosGraphValues = { FVector2D(0, FMath::Cos(0.0)) };
diff --git a/Engine/Plugins/Experimental/SlateIM/Source/SlateIM/Private/Misc/SlateIMManager.h b/Engine/Plugins/Experimental/SlateIM/Source/SlateIM/Private/Misc/SlateIMManager.h
--- a/Engine/Plugins/Experimental/SlateIM/Source/SlateIM/Private/Misc/SlateIMManager.h
+++ b/Engine/Plugins/Experimental/SlateIM/Source/SlateIM/Private/Misc/SlateIMManager.h
@@ -128,6 +128,8 @@ namespace SlateIM
 		void PopContainer();
 		template<typename ContainerType>
 		TSharedPtr<ContainerType> GetCurrentContainer() const;
+		template<typename ContainerType>
+		TSharedPtr<ContainerType> FindMostRecentContainer() const;
 		const FContainerNode* GetCurrentContainerNode() const;
 		
 		void PushMenuRoot(TSharedPtr<SImContextMenuAnchor>& MenuRoot);
@@ -254,6 +256,23 @@ namespace SlateIM
 		return nullptr;
 	}
 
+	template <typename ContainerType>
+	TSharedPtr<ContainerType> FSlateIMManager::FindMostRecentContainer() const
+	{
+		static_assert(std::is_base_of_v<ISlateIMContainer, ContainerType>, "ContainerType must derive from ISlateIMContainer");
+
+		for (int32 i = CurrentRoot->CurrentContainerStack.Num() - 1; i >= 0; --i)
+		{
+			TSharedPtr<ISlateIMContainer> Container = CurrentRoot->CurrentContainerStack[i].Widget;
+			if (Container && Container->IsA<ContainerType>())
+			{
+				return StaticCastSharedPtr<ContainerType>(Container);
+			}
+		}
+
+		return nullptr;
+	}
+
 	template <typename WidgetType>
 	TSharedPtr<WidgetType> FSlateIMManager::BeginIMWidget()
 	{
diff --git a/Engine/Plugins/Experimental/SlateIM/Source/SlateIM/Private/Roots/SlateIMWindowRoot.cpp b/Engine/Plugins/Experimental/SlateIM/Source/SlateIM/Private/Roots/SlateIMWindowRoot.cpp
--- a/Engine/Plugins/Experimental/SlateIM/Source/SlateIM/Private/Roots/SlateIMWindowRoot.cpp
+++ b/Engine/Plugins/Experimental/SlateIM/Source/SlateIM/Private/Roots/SlateIMWindowRoot.cpp
@@ -56,3 +56,8 @@ void FSlateIMWindowRoot::UpdateWindow(const FStringView& Title)
 		}
 	}
 }
+
+TSharedPtr<SWindow> FSlateIMWindowRoot::GetWindow() const
+{
+	return RootWindow.Pin();
+}
diff --git a/Engine/Plugins/Experimental/SlateIM/Source/SlateIM/Private/Roots/SlateIMWindowRoot.h b/Engine/Plugins/Experimental/SlateIM/Source/SlateIM/Private/Roots/SlateIMWindowRoot.h
--- a/Engine/Plugins/Experimental/SlateIM/Source/SlateIM/Private/Roots/SlateIMWindowRoot.h
+++ b/Engine/Plugins/Experimental/SlateIM/Source/SlateIM/Private/Roots/SlateIMWindowRoot.h
@@ -20,6 +20,8 @@ public:
 	virtual FSlateIMInputState& GetInputState() override;
 
 	void UpdateWindow(const FStringView& Title);
+
+	TSharedPtr<SWindow> GetWindow() const;
 	
 private:
 	TWeakPtr<SWindow> RootWindow;
diff --git a/Engine/Plugins/Experimental/SlateIM/Source/SlateIM/Public/SlateIM.h b/Engine/Plugins/Experimental/SlateIM/Source/SlateIM/Public/SlateIM.h
--- a/Engine/Plugins/Experimental/SlateIM/Source/SlateIM/Public/SlateIM.h
+++ b/Engine/Plugins/Experimental/SlateIM/Source/SlateIM/Public/SlateIM.h
@@ -6,6 +6,7 @@
 #include "Layout/Margin.h"
 #include "Modules/ModuleInterface.h"
 #include "Styling/SlateColor.h"
+#include "Textures/SlateIcon.h"
 #include "Types/SlateEnums.h"
 #include "Widgets/Layout/Anchors.h"
 
@@ -593,6 +594,65 @@ namespace SlateIM
 	SLATEIM_API void EndSubMenu();
 #pragma endregion Menu
 
+#pragma region Tabs
+	/**
+	 * Begin a tab group, can contain TabStacks and TabSplitters
+	 * 
+	 * @param TabGroupId Provide a unique identifier for this tab group
+	 */
+	SLATEIM_API void BeginTabGroup(const FName& TabGroupId);
+	
+	/**
+	 * Ends the current Tab Group
+	 */
+	SLATEIM_API void EndTabGroup();
+
+	/**
+	 * Begin a tab stack, can only contain Tabs
+	 */
+	SLATEIM_API void BeginTabStack();
+
+	/**
+	 * Ends the current Tab Stack
+	 */
+	SLATEIM_API void EndTabStack();
+
+	/**
+	 * Begin a tab splitter, displays child TabSplitters and TabStacks side-by-side
+	 * 
+	 * @param Orientation The direction to layout child TabSplitters and TabStacks
+	 */
+	SLATEIM_API void BeginTabSplitter(EOrientation Orientation);
+
+	/**
+	 * Ends the current Tab Splitter
+	 */
+	SLATEIM_API void EndTabSplitter();
+
+	/**
+	 * Assigns the Size Coefficient for the next child of a TabSplitter
+	 * 
+	 * @param SizeCoefficient The weight to assign to the next child of the parent TabSplitter
+	 */
+	SLATEIM_API void TabSplitterSizeCoefficient(float SizeCoefficient);
+
+	/**
+	 * Begins a Tab to contain any other content
+	 * 
+	 * @param TabId Provide a tab id that is unique within the Tab Group
+	 * @param TabIcon (optional) An icon to display on the tab next to the title. This cannot be updated after the initial creation of the tab.
+	 * @param TabTitle (optional) The title to display for the tab if it should differ from the TabId. This cannot be updated after the initial creation of the tab.
+	 * 
+	 * @return Returns true if the tab is active, false otherwise. Logic to draw the contents of the tab can be skipped when this is false.
+	 */
+	SLATEIM_API bool BeginTab(const FName& TabId, const FSlateIcon& TabIcon = FSlateIcon(), const FText& TabTitle = FText::GetEmpty());
+
+	/**
+	 * Ends the current Tab. This must be called even when BeginTab returns false.
+	 */
+	SLATEIM_API void EndTab();
+#pragma endregion Tabs
+
 #pragma region Util
 	/**
 	 * SlateIM updates are disabled in specific scenarios (like when a SlateIM::ModalDialog is open), use this function to react accordingly
