SlateIM: Fix Android build issue

#rnx

diff --git a/Engine/Plugins/Experimental/SlateIM/Source/SlateIM/Private/Containers/SImTabGroup.cpp b/Engine/Plugins/Experimental/SlateIM/Source/SlateIM/Private/Containers/SImTabGroup.cpp
--- a/Engine/Plugins/Experimental/SlateIM/Source/SlateIM/Private/Containers/SImTabGroup.cpp
+++ b/Engine/Plugins/Experimental/SlateIM/Source/SlateIM/Private/Containers/SImTabGroup.cpp
@@ -140,7 +140,7 @@ void SImTabGroup::Construct(const FArguments& InArgs)
 	const TSharedPtr<SDockTab> ParentTab = SlateIM::FSlateIMManager::Get().FindMostRecentContainer<SImTab>();
 	TabManager = FGlobalTabmanager::Get()->NewTabManager(ParentTab.IsValid() ? ParentTab.ToSharedRef() : SNew(SDockTab));
 	const TSharedRef<FTabManager::FArea> PrimaryArea = FTabManager::NewPrimaryArea();
-	WorkingLayoutNodeStack.Emplace(PrimaryArea);
+	WorkingLayoutNodeStack.Emplace(FImTabLayoutNode{ PrimaryArea });
 	TabLayout = FTabManager::NewLayout(TabGroupId)->AddArea(PrimaryArea);
 }
 
@@ -323,7 +323,7 @@ void SImTabGroup::BeginTabSplitter(EOrientation Orientation)
 
 	// Reset the size coefficient after using it
 	CurrentSizeCoefficient = 1.f;
-	WorkingLayoutNodeStack.Emplace(TabSplitter);
+	WorkingLayoutNodeStack.Emplace(FImTabLayoutNode{ TabSplitter });
 }
 
 void SImTabGroup::BeginTabStack()
@@ -350,7 +350,7 @@ void SImTabGroup::BeginTabStack()
 	
 	// Reset the size coefficient after using it
 	CurrentSizeCoefficient = 1.f;
-	WorkingLayoutNodeStack.Emplace(TabStack);
+	WorkingLayoutNodeStack.Emplace(FImTabLayoutNode{ TabStack });
 }
 
 void SImTabGroup::EndTabLayoutNode()
