[SlateIM] Fix initially selected index not showing as such when SlateIM::ComboBox is spawned

diff --git a/Engine/Plugins/Experimental/SlateIM/Source/SlateIM/Private/Implementation/SlateIM_Widgets.cpp b/Engine/Plugins/Experimental/SlateIM/Source/SlateIM/Private/Implementation/SlateIM_Widgets.cpp
--- a/Engine/Plugins/Experimental/SlateIM/Source/SlateIM/Private/Implementation/SlateIM_Widgets.cpp
+++ b/Engine/Plugins/Experimental/SlateIM/Source/SlateIM/Private/Implementation/SlateIM_Widgets.cpp
@@ -590,6 +590,7 @@ namespace SlateIM
 			ComboWidget =
 				SNew(STextComboBox)
 				.OptionsSource(&WidgetIMData->Data)
+				.InitiallySelectedItem(WidgetIMData->Data.IsValidIndex(InOutSelectedItemIndex) ? WidgetIMData->Data[InOutSelectedItemIndex] : nullptr)
 				.OnSelectionChanged_Lambda([ActivationData = Scope.GetOrCreateActivationMetadata().ToWeakPtr()](TSharedPtr<FString> NewVal, ESelectInfo::Type SelectType)
 				{
 					FSlateIMManager::Get().ActivateWidget(ActivationData.Pin());
