SlateIM: Fix Viewport roots preventing mouse input for the whole viewport

Now viewport widgets will only handle mouse inputs in their visible area.

Also added Mouse inputs to the Input tab of the example widget.

diff --git a/Engine/Plugins/Experimental/SlateIM/Source/SlateIM/Private/Misc/SlateIMExamples.cpp b/Engine/Plugins/Experimental/SlateIM/Source/SlateIM/Private/Misc/SlateIMExamples.cpp
--- a/Engine/Plugins/Experimental/SlateIM/Source/SlateIM/Private/Misc/SlateIMExamples.cpp
+++ b/Engine/Plugins/Experimental/SlateIM/Source/SlateIM/Private/Misc/SlateIMExamples.cpp
@@ -97,7 +97,8 @@ void FSlateIMTestWidget::Draw()
 	TimeSinceLastUpdate += FApp::GetDeltaTime();
 
 	SCOPED_NAMED_EVENT_TEXT("FSlateIMTestWidget::Draw", FColorList::Goldenrod);
-	SlateIM::BeginBorder(FAppStyle::GetBrush("ToolPanel.GroupBorder"));
+	constexpr bool bAbsorbMouse = false;
+	SlateIM::BeginBorder(FAppStyle::GetBrush("ToolPanel.GroupBorder"), Orient_Vertical, bAbsorbMouse);
 	// Basic perf measurement, outside of the scrollbox so that it's "pinned"
 	if (TimeSinceLastUpdate > 0.5f)
 	{
@@ -901,6 +902,50 @@ void FSlateIMTestWidget::DrawInputs()
 		SlateIM::EndVerticalStack();
 	}
 
+	{
+		SlateIM::BeginVerticalStack();
+		SlateIM::Text(TEXT("Mouse X-value"));
+		const float NormalizedAnalogXValue = (1.f + (SlateIM::GetKeyAnalogValue(EKeys::MouseX) / 100.f)) * 0.5f;
+		SlateIM::VAlign(VAlign_Center);
+		SlateIM::ProgressBar(NormalizedAnalogXValue);
+	
+		SlateIM::Fill();
+		SlateIM::Spacer(FVector2D(1.f, 1.f));
+	
+		SlateIM::Text(TEXT("Mouse Y-value"));
+		const float NormalizedAnalogYValue = (1.f + (SlateIM::GetKeyAnalogValue(EKeys::MouseY) / 100.f)) * 0.5f;
+		SlateIM::VAlign(VAlign_Center);
+		SlateIM::ProgressBar(NormalizedAnalogYValue);
+
+		SlateIM::BeginHorizontalStack();
+		{
+			LMBBrush = FSlateColorBrush(SlateIM::Private::GetKeyStateColor(EKeys::LeftMouseButton));
+			SlateIM::MinWidth(50.f);
+			SlateIM::MinHeight(50.f);
+			SlateIM::MaxWidth(50.f);
+			SlateIM::MaxHeight(50.f);
+			SlateIM::BeginBorder(&LMBBrush);
+			SlateIM::HAlign(HAlign_Center);
+			SlateIM::VAlign(VAlign_Center);
+			SlateIM::Text(TEXT("LMB"), FLinearColor::White);
+			SlateIM::EndBorder();
+		}
+		{
+			RMBBrush = FSlateColorBrush(SlateIM::Private::GetKeyStateColor(EKeys::RightMouseButton));
+			SlateIM::MinWidth(50.f);
+			SlateIM::MinHeight(50.f);
+			SlateIM::MaxWidth(50.f);
+			SlateIM::MaxHeight(50.f);
+			SlateIM::BeginBorder(&RMBBrush);
+			SlateIM::HAlign(HAlign_Center);
+			SlateIM::VAlign(VAlign_Center);
+			SlateIM::Text(TEXT("RMB"), FLinearColor::White);
+			SlateIM::EndBorder();
+		}
+		SlateIM::EndHorizontalStack();
+		SlateIM::EndVerticalStack();
+	}
+
 	{
 		SlateIM::BeginVerticalStack();
 		SlateIM::Text(TEXT("Right Stick X-value"));
diff --git a/Engine/Plugins/Experimental/SlateIM/Source/SlateIM/Private/Misc/SlateIMExamples.h b/Engine/Plugins/Experimental/SlateIM/Source/SlateIM/Private/Misc/SlateIMExamples.h
--- a/Engine/Plugins/Experimental/SlateIM/Source/SlateIM/Private/Misc/SlateIMExamples.h
+++ b/Engine/Plugins/Experimental/SlateIM/Source/SlateIM/Private/Misc/SlateIMExamples.h
@@ -53,6 +53,9 @@ private:
 	FSlateBrush ABrush;
 	FSlateBrush SBrush;
 	FSlateBrush DBrush;
+	
+	FSlateBrush LMBBrush;
+	FSlateBrush RMBBrush;
 
 #if WITH_ENGINE
 	TSoftObjectPtr<UTexture2D> RedIcon = TSoftObjectPtr<UTexture2D>(FSoftObjectPath(TEXT("/Engine/EngineResources/AICON-Red.AICON-Red")));
diff --git a/Engine/Plugins/Experimental/SlateIM/Source/SlateIM/Private/Roots/SlateIMViewportRoot.cpp b/Engine/Plugins/Experimental/SlateIM/Source/SlateIM/Private/Roots/SlateIMViewportRoot.cpp
--- a/Engine/Plugins/Experimental/SlateIM/Source/SlateIM/Private/Roots/SlateIMViewportRoot.cpp
+++ b/Engine/Plugins/Experimental/SlateIM/Source/SlateIM/Private/Roots/SlateIMViewportRoot.cpp
@@ -62,30 +62,32 @@ FSlateIMViewportRoot::~FSlateIMViewportRoot()
 
 void FSlateIMViewportRoot::UpdateChild(TSharedRef<SWidget> Child, const FSlateIMSlotData& AlignmentData)
 {
-	if (Slot)
+	if (Wrapper)
 	{
-		Slot->AttachWidget(Child);
+		Wrapper->SetContent(Child);
 	}
 }
 
 FSlateIMInputState& FSlateIMViewportRoot::GetInputState()
 {
-	check(ViewportRoot.IsValid());
-	return ViewportRoot->InputState;
+	check(Wrapper.IsValid());
+	return Wrapper->InputState;
 }
 
 void FSlateIMViewportRoot::UpdateViewport(const SlateIM::FViewportRootLayout& Layout)
 {
 	if (!ViewportRoot)
 	{
-		ViewportRoot = SNew(SImWrapper)
-		[
+		ViewportRoot = 
 			SNew(SConstraintCanvas)
+			.Visibility(EVisibility::SelfHitTestInvisible)
 			+ SConstraintCanvas::Slot()
 			.Anchors(Layout.Anchors)
 			.Alignment(Layout.Alignment)
 			.Expose(Slot)
-		];
+			[
+				SAssignNew(Wrapper, SImWrapper)	
+			];
 		
 #if WITH_EDITOR
 		if (TSharedPtr<IAssetViewport> AssetViewportPinned = AssetViewport.Pin())
diff --git a/Engine/Plugins/Experimental/SlateIM/Source/SlateIM/Private/Roots/SlateIMViewportRoot.h b/Engine/Plugins/Experimental/SlateIM/Source/SlateIM/Private/Roots/SlateIMViewportRoot.h
--- a/Engine/Plugins/Experimental/SlateIM/Source/SlateIM/Private/Roots/SlateIMViewportRoot.h
+++ b/Engine/Plugins/Experimental/SlateIM/Source/SlateIM/Private/Roots/SlateIMViewportRoot.h
@@ -49,7 +49,8 @@ public:
 #endif
 
 private:
-	TSharedPtr<SImWrapper> ViewportRoot;
+	TSharedPtr<SImWrapper> Wrapper;
+	TSharedPtr<SWidget> ViewportRoot;
 	SConstraintCanvas::FSlot* Slot = nullptr;
 };
 #endif // WITH_ENGINE
