SlateIM: Clarify messaging for assert when a Root is started without calling SlateIM::EndRoot()

#rnx

diff --git a/Engine/Plugins/Experimental/SlateIM/Source/SlateIM/Private/Misc/SlateIMManager.cpp b/Engine/Plugins/Experimental/SlateIM/Source/SlateIM/Private/Misc/SlateIMManager.cpp
--- a/Engine/Plugins/Experimental/SlateIM/Source/SlateIM/Private/Misc/SlateIMManager.cpp
+++ b/Engine/Plugins/Experimental/SlateIM/Source/SlateIM/Private/Misc/SlateIMManager.cpp
@@ -359,7 +359,7 @@ namespace SlateIM
 			return;
 		}
 
-		check(CurrentRoot == nullptr);
+		checkf(CurrentRoot == nullptr, TEXT("There's an active SlateIM Root, are you missing a call to SlateIM::EndRoot()?"));
 
 		ActivatedWidgets.Reset();
 
