SlateIM: Fix case where child rows would not get removed when no longer drawn

diff --git a/Engine/Plugins/Experimental/SlateIM/Source/SlateIM/Private/Containers/SImTableView.cpp b/Engine/Plugins/Experimental/SlateIM/Source/SlateIM/Private/Containers/SImTableView.cpp
--- a/Engine/Plugins/Experimental/SlateIM/Source/SlateIM/Private/Containers/SImTableView.cpp
+++ b/Engine/Plugins/Experimental/SlateIM/Source/SlateIM/Private/Containers/SImTableView.cpp
@@ -18,7 +18,17 @@ namespace SlateIM::Private
 
 void FSlateIMTableRow::RemoveUnusedChildren(int32 LastUsedChildIndex)
 {
-	Children.SetNum(FMath::Min(Children.Num(), LastUsedChildIndex + 1));
+	const int32 NumChildren = FMath::Min(Children.Num(), LastUsedChildIndex + 1);
+	if (Children.Num() != NumChildren)
+	{
+		Children.SetNum(NumChildren);
+
+		// Modifying the child count requires a rebuild to handle the case where we need to create/destroy the expander arrow and child widgets
+		if (OwningTable)
+		{
+			OwningTable->RequestRebuild();
+		}
+	}
 }
 
 void FSlateIMTableRow::UpdateChild(FSlateIMChild Child, int32 Index, const FSlateIMSlotData& AlignmentData)
@@ -38,9 +48,10 @@ void FSlateIMTableRow::UpdateChild(FSlateIMChild Child, int32 Index, const FSlat
 		Children.Add(Child);
 	}
 
+	// Modifying the child count requires a rebuild to handle the case where we need to create/destroy the expander arrow and child widgets
 	if (OwningTable)
 	{
-		OwningTable->MarkDirty();
+		OwningTable->RequestRebuild();
 	}
 }
 
@@ -251,7 +262,7 @@ void SImTableView::RemoveUnusedChildren(int32 LastUsedChildIndex)
 {
 	if (TableRows.Num() > (LastUsedChildIndex + 1))
 	{
-		bIsDirty = true;
+		RequestRefresh();
 		TableRows.SetNum(LastUsedChildIndex + 1);
 	}
 }
@@ -270,19 +281,19 @@ void SImTableView::UpdateChild(FSlateIMChild Child, int32 Index, const FSlateIMS
 	{
 		if (TableRows[Index] != Row)
 		{
-			bIsDirty = true;
+			RequestRefresh();
 			TableRows[Index] = Row.ToSharedRef();
 		}
 	}
 	else
 	{
-		bIsDirty = true;
+		RequestRefresh();
 		TableRows.Add(Row.ToSharedRef());
 	}
 
 	if (Row->GetColumnCount() != ColumnCount)
 	{
-		bIsDirty = true;
+		RequestRefresh();
 		Row->UpdateColumnCount(ColumnCount);
 	}
 }
@@ -325,11 +336,16 @@ void SImTableView::BeginTableUpdates()
 
 void SImTableView::EndTableUpdates()
 {
-	if (bIsDirty)
+	if (DirtyState == NeedsRefresh)
 	{
-		bIsDirty = false;
+		DirtyState = Clean;
 		RequestListRefresh();
 	}
+	else if (DirtyState == NeedsRebuild)
+	{
+		DirtyState = Clean;
+		RebuildList();
+	}
 }
 
 void SImTableView::BeginTableContent()
@@ -373,6 +389,16 @@ bool SImTableView::IsTree() const
 	return false;
 }
 
+void SImTableView::RequestRefresh()
+{
+	DirtyState = NeedsRefresh;
+}
+
+void SImTableView::RequestRebuild()
+{
+	DirtyState = NeedsRebuild;
+}
+
 TSharedRef<ITableRow> SImTableView::GenerateRow(TSharedRef<FSlateIMTableRow> InTableRow, const TSharedRef<STableViewBase>& OwnerTable)
 {
 	return SNew(SImTableRow, SharedThis(this), InTableRow)
@@ -422,5 +448,5 @@ void SImTableView::AddColumn_Internal(const FName& ColumnId, const FStringView&
 	}
 		
 	Header->AddColumn(Args);
-	bIsDirty = true;
+	RequestRefresh();
 }
diff --git a/Engine/Plugins/Experimental/SlateIM/Source/SlateIM/Private/Containers/SImTableView.h b/Engine/Plugins/Experimental/SlateIM/Source/SlateIM/Private/Containers/SImTableView.h
--- a/Engine/Plugins/Experimental/SlateIM/Source/SlateIM/Private/Containers/SImTableView.h
+++ b/Engine/Plugins/Experimental/SlateIM/Source/SlateIM/Private/Containers/SImTableView.h
@@ -95,7 +95,8 @@ public:
 	// A table is a tree if any of its rows has children
 	bool IsTree() const;
 
-	void MarkDirty() { bIsDirty = true; }
+	void RequestRefresh();
+	void RequestRebuild();
 
 	TSharedRef<ITableRow> GenerateRow(TSharedRef<FSlateIMTableRow> InTableRow, const TSharedRef<STableViewBase>& OwnerTable);
 
@@ -110,9 +111,15 @@ private:
 
 	TArray<TSharedRef<FSlateIMTableRow>> TableRows;
 	TSharedPtr<SImTableHeader> Header;
-	int32 ColumnCount = 0;
-	bool bIsDirty = false;
 	const FTableRowStyle* RowStyle = nullptr;
+	int32 ColumnCount = 0;
 
 	mutable int32 CachedNumLiveWidgets = 0;
+
+	enum
+	{
+		Clean,
+		NeedsRefresh,
+		NeedsRebuild
+	} DirtyState = NeedsRebuild;
 };
diff --git a/Engine/Plugins/Experimental/SlateIM/Source/SlateIM/Private/Misc/SlateIMExamples.cpp b/Engine/Plugins/Experimental/SlateIM/Source/SlateIM/Private/Misc/SlateIMExamples.cpp
--- a/Engine/Plugins/Experimental/SlateIM/Source/SlateIM/Private/Misc/SlateIMExamples.cpp
+++ b/Engine/Plugins/Experimental/SlateIM/Source/SlateIM/Private/Misc/SlateIMExamples.cpp
@@ -766,6 +766,42 @@ void FSlateIMTestWidget::DrawTables()
 		
 			SlateIM::EndTable();
 		}
+
+		// Dynamically added/removed child row
+		{
+			SlateIM::CheckBox(TEXT("Show Child?"), bShouldAddChildRow);
+
+			SlateIM::Fill();
+			SlateIM::HAlign(HAlign_Fill);
+			SlateIM::BeginTable();
+			SlateIM::AddTableColumn(TEXT("Column"));
+
+			if (SlateIM::NextTableCell())
+			{
+				SlateIM::HAlign(HAlign_Fill);
+				SlateIM::Text(TEXT("First Row"));
+			}
+		
+			if (SlateIM::NextTableCell())
+			{
+				SlateIM::HAlign(HAlign_Fill);
+				SlateIM::Text(TEXT("Second Row"));
+			}
+
+			if (bShouldAddChildRow)
+			{
+				if (SlateIM::BeginTableRowChildren())
+				{
+					if (SlateIM::NextTableCell())
+					{
+						SlateIM::HAlign(HAlign_Fill);
+						SlateIM::Text(TEXT("Child Row"));
+					}
+				}
+				SlateIM::EndTableRowChildren();
+			}
+			SlateIM::EndTable();
+		}
 	}
 	SlateIM::EndScrollBox();
 }
diff --git a/Engine/Plugins/Experimental/SlateIM/Source/SlateIM/Private/Misc/SlateIMExamples.h b/Engine/Plugins/Experimental/SlateIM/Source/SlateIM/Private/Misc/SlateIMExamples.h
--- a/Engine/Plugins/Experimental/SlateIM/Source/SlateIM/Private/Misc/SlateIMExamples.h
+++ b/Engine/Plugins/Experimental/SlateIM/Source/SlateIM/Private/Misc/SlateIMExamples.h
@@ -41,6 +41,7 @@ private:
 	int32 LiveNumItems = NumItems;
 	FString NumItemsText = FString::FromInt(NumItems);
 	bool bShouldLiveUpdateTable = false;
+	bool bShouldAddChildRow = false;
 	TArray<FString> ComboBoxItems = { TEXT("Option 1"), TEXT("Option 2"), TEXT("Option 3"), TEXT("Option 4") };
 	bool bRefreshComboItems = false;
 	int32 DynamicTabCount = 3;
