SlateIM: Fix example window text not readjusting itself after the window is resized

Due to a bug in slate's auto wrap text layout, text blocks won't increase in size after wrapping unless the textblocks fill the width of the resized container.

#rnx

diff --git a/Engine/Plugins/Experimental/SlateIM/Source/SlateIM/Private/Misc/SlateIMExamples.cpp b/Engine/Plugins/Experimental/SlateIM/Source/SlateIM/Private/Misc/SlateIMExamples.cpp
--- a/Engine/Plugins/Experimental/SlateIM/Source/SlateIM/Private/Misc/SlateIMExamples.cpp
+++ b/Engine/Plugins/Experimental/SlateIM/Source/SlateIM/Private/Misc/SlateIMExamples.cpp
@@ -368,8 +368,11 @@ void FSlateIMTestWidget::DrawBasics()
 		// ContextMenu examples
 		{
 			SCOPED_NAMED_EVENT_TEXT("ContextMenu examples", FColorList::Goldenrod);
+			SlateIM::HAlign(HAlign_Fill);
 			SlateIM::Text(TEXT("Context Menu Test"));
+			SlateIM::HAlign(HAlign_Fill);
 			SlateIM::BeginContextMenuAnchor();
+			SlateIM::HAlign(HAlign_Fill);
 			SlateIM::Text(TEXT("Right Click here to show a menu"));
 
 			// This part is only shown if the menu is open
@@ -404,6 +407,7 @@ void FSlateIMTestWidget::DrawBasics()
 		// Modal Examples
 		{
 			SCOPED_NAMED_EVENT_TEXT("Modal Examples", FColorList::Goldenrod);
+			SlateIM::HAlign(HAlign_Fill);
 			SlateIM::Text(TEXT("Open Modal Dialog of Type:"));
 
 			// Wrap example
@@ -445,7 +449,9 @@ void FSlateIMTestWidget::DrawBasics()
 
 			if (DialogResult.IsSet())
 			{
+				SlateIM::HAlign(HAlign_Fill);
 				SlateIM::Text(TEXT("Dialog Result:"));
+				SlateIM::HAlign(HAlign_Fill);
 				switch (DialogResult.GetValue())
 				{
 				case EAppReturnType::No:
@@ -484,6 +490,7 @@ void FSlateIMTestWidget::DrawBasics()
 			}
 			else
 			{
+				SlateIM::HAlign(HAlign_Fill);
 				SlateIM::Text(TEXT("No Dialog Result"));
 			}
 		}
