SlateIM: Add an overload of SlateIM::Image that takes just a color

#rnx

diff --git a/Engine/Plugins/Experimental/SlateIM/Source/SlateIM/Private/Implementation/SlateIM_Widgets.cpp b/Engine/Plugins/Experimental/SlateIM/Source/SlateIM/Private/Implementation/SlateIM_Widgets.cpp
--- a/Engine/Plugins/Experimental/SlateIM/Source/SlateIM/Private/Implementation/SlateIM_Widgets.cpp
+++ b/Engine/Plugins/Experimental/SlateIM/Source/SlateIM/Private/Implementation/SlateIM_Widgets.cpp
@@ -175,6 +175,12 @@ namespace SlateIM
 		Image(FAppStyle::Get().GetBrush(ImageStyleName), ColorAndOpacity, DesiredSize);
 	}
 
+	void Image(const FSlateColor& ColorAndOpacity, FVector2D DesiredSize)
+	{
+		static FSlateBrush DefaultBrush;
+		Image(&DefaultBrush, ColorAndOpacity, DesiredSize);
+	}
+
 #if WITH_ENGINE
 	struct FPinnedImageResource
 	{
diff --git a/Engine/Plugins/Experimental/SlateIM/Source/SlateIM/Public/SlateIM.h b/Engine/Plugins/Experimental/SlateIM/Source/SlateIM/Public/SlateIM.h
--- a/Engine/Plugins/Experimental/SlateIM/Source/SlateIM/Public/SlateIM.h
+++ b/Engine/Plugins/Experimental/SlateIM/Source/SlateIM/Public/SlateIM.h
@@ -370,6 +370,13 @@ namespace SlateIM
 	 * @param DesiredSize (optional) Override the desired size to display the brush at
 	 */
 	SLATEIM_API void Image(const FName ImageStyleName, const FSlateColor& ColorAndOpacity = FLinearColor::White, FVector2D DesiredSize = FVector2D::ZeroVector);
+	/**
+	 * Display a colored box
+	 * 
+	 * @param ColorAndOpacity The color to display
+	 * @param DesiredSize (optional) Override the desired size to display the box at
+	 */
+	SLATEIM_API void Image(const FSlateColor& ColorAndOpacity, FVector2D DesiredSize = FVector2D::ZeroVector);
 #if WITH_ENGINE
 	/**
 	 * Display a texture
