SlateIM InGame widget actor for server/client debugging

diff --git a/Engine/Plugins/Experimental/SlateIM/SlateIM.uplugin b/Engine/Plugins/Experimental/SlateIM/SlateIM.uplugin
--- a/Engine/Plugins/Experimental/SlateIM/SlateIM.uplugin
+++ b/Engine/Plugins/Experimental/SlateIM/SlateIM.uplugin
@@ -20,6 +20,11 @@
 			"Name": "SlateIM",
 			"Type": "Runtime",
 			"LoadingPhase": "Default"
+		},
+		{
+			"Name": "SlateIMInGame",
+			"Type": "Runtime",
+			"LoadingPhase": "Default"
 		}
 	]
 }
\ No newline at end of file
diff --git a/Engine/Plugins/Experimental/SlateIM/Source/SlateIMInGame/Private/SlateIMInGameWidgetBase.cpp b/Engine/Plugins/Experimental/SlateIM/Source/SlateIMInGame/Private/SlateIMInGameWidgetBase.cpp
new file mode 100644
--- /dev/null
+++ b/Engine/Plugins/Experimental/SlateIM/Source/SlateIMInGame/Private/SlateIMInGameWidgetBase.cpp
@@ -0,0 +1,194 @@
+﻿// Copyright Epic Games, Inc. All Rights Reserved.
+
+#include "SlateIMInGameWidgetBase.h"
+
+#include "EngineUtils.h"
+#include "Framework/Application/SlateApplication.h"
+#include "GameFramework/PlayerController.h"
+#include "SlateIM.h"
+
+ASlateIMInGameWidgetBase::ASlateIMInGameWidgetBase()
+{
+	PrimaryActorTick.bCanEverTick = false;
+	PrimaryActorTick.bStartWithTickEnabled = false;
+
+	bReplicates = true;
+	bAlwaysRelevant = true;
+
+	SetNetUpdateFrequency(3.f);
+	SetMinNetUpdateFrequency(3.f);
+}
+
+void ASlateIMInGameWidgetBase::BeginPlay()
+{
+	Super::BeginPlay();
+
+	if (IsLocallyControlled())
+	{
+		StartWidget();
+	}
+}
+
+void ASlateIMInGameWidgetBase::EndPlay(EEndPlayReason::Type Reason)
+{
+	if (IsLocallyControlled())
+	{
+		StopWidget();
+	}
+
+	Super::EndPlay(Reason);
+}
+
+void ASlateIMInGameWidgetBase::OnRep_Owner()
+{
+	Super::OnRep_Owner();
+
+	if (HasActorBegunPlay())
+	{
+		if (IsLocallyControlled())
+		{
+			StartWidget();
+		}
+		else
+		{
+			StopWidget();
+		}
+	}
+}
+
+void ASlateIMInGameWidgetBase::PreReplication(IRepChangedPropertyTracker& ChangedPropertyTracker)
+{
+	Super::PreReplication(ChangedPropertyTracker);
+
+#if WITH_SERVER_CODE
+	if (HasAuthority())
+	{
+		GenerateServerSnapshot();
+	}
+#endif // WITH_SERVER_CODE
+}
+
+bool ASlateIMInGameWidgetBase::IsNetRelevantFor(const AActor* RealViewer, const AActor* ViewTarget, const FVector& SrcLocation) const
+{
+	return GetOwner() != nullptr && GetOwner() == RealViewer;
+}
+
+APlayerController* ASlateIMInGameWidgetBase::GetPlayerController() const
+{
+	return Cast<APlayerController>(GetOwner());
+}
+
+bool ASlateIMInGameWidgetBase::IsLocallyControlled() const
+{
+	if (APlayerController* OwnerPlayerController = GetPlayerController())
+	{
+		return OwnerPlayerController->IsLocalController();
+	}
+
+	return false;
+}
+
+ASlateIMInGameWidgetBase* ASlateIMInGameWidgetBase::GetInGameWidget(const APlayerController* Owner, const TSubclassOf<ASlateIMInGameWidgetBase>& InGameWidgetClass)
+{
+	if (ensure(Owner && InGameWidgetClass))
+	{
+		for (TActorIterator<ASlateIMInGameWidgetBase> It(Owner->GetWorld(), InGameWidgetClass); It; ++It)
+		{
+			if (It->GetOwner() == Owner)
+			{
+				return *It;
+			}
+		}
+	}
+	return nullptr;
+}
+
+void ASlateIMInGameWidgetBase::EnableInGameWidget(APlayerController* Owner, const bool bEnable, const TSubclassOf<ASlateIMInGameWidgetBase>& InGameWidgetClass)
+{
+	if (bEnable)
+	{
+		GetOrOpenInGameWidget(Owner, InGameWidgetClass);
+	}
+	else 
+	{
+		DestroyInGameWidget(Owner, InGameWidgetClass);
+	}
+}
+
+void ASlateIMInGameWidgetBase::StartWidget()
+{
+	if (!WidgetTickHandle.IsValid())
+	{
+		WidgetTickHandle = FSlateApplication::Get().OnPreTick().AddUObject(this, &ASlateIMInGameWidgetBase::TickWidget);
+	}
+}
+
+void ASlateIMInGameWidgetBase::StopWidget()
+{
+	if (WidgetTickHandle.IsValid())
+	{
+		FSlateApplication::Get().OnPreTick().Remove(WidgetTickHandle);
+		WidgetTickHandle.Reset();
+	}
+}
+
+void ASlateIMInGameWidgetBase::TickWidget(const float DeltaTime)
+{
+	if (!SlateIM::CanUpdateSlateIM())
+	{
+		return;
+	}
+
+	DrawWidget(DeltaTime);
+}
+
+ASlateIMInGameWidgetBase* ASlateIMInGameWidgetBase::GetOrOpenInGameWidget(APlayerController* Owner, const TSubclassOf<ASlateIMInGameWidgetBase>& InGameWidgetClass)
+{
+	if (ensure(Owner && InGameWidgetClass))
+	{
+		if (ASlateIMInGameWidgetBase* ExistingWidget = GetInGameWidget(Owner, InGameWidgetClass))
+		{
+			return ExistingWidget;
+		}
+
+		if (GetDefault<ASlateIMInGameWidgetBase>(InGameWidgetClass)->bReplicates && Owner->GetWorld()->GetNetMode() == ENetMode::NM_Client)
+		{
+			UE_LOG(LogActor, Error, TEXT("In game widget %s was requested but it's replicated and needs server permission."), *InGameWidgetClass->GetName());
+		}
+		else
+		{
+			FActorSpawnParameters Params;
+			Params.Owner = Owner;
+
+			if (ASlateIMInGameWidgetBase* NewWidget = Owner->GetWorld()->SpawnActor<ASlateIMInGameWidgetBase>(InGameWidgetClass, Params))
+			{
+				NewWidget->AttachToActor(Owner, FAttachmentTransformRules::KeepRelativeTransform);
+			}
+		}
+	}
+
+	return nullptr;
+}
+
+void ASlateIMInGameWidgetBase::DestroyInGameWidget(const APlayerController* Owner, const TSubclassOf<ASlateIMInGameWidgetBase>& InGameWidgetClass)
+{
+	if (ensure(Owner && InGameWidgetClass))
+	{
+		if (ASlateIMInGameWidgetBase* OpenedWidget = GetInGameWidget(Owner, InGameWidgetClass))
+		{
+			if (GetDefault<ASlateIMInGameWidgetBase>(InGameWidgetClass)->bReplicates && Owner->GetWorld()->GetNetMode() == ENetMode::NM_Client)
+			{
+				UE_LOG(LogActor, Error, TEXT("Destruction of in game widget %s was requested but it is replicated and needs server permission."), *InGameWidgetClass->GetName());
+			}
+			else
+			{
+				OpenedWidget->Destroy();
+			}
+		}
+	}
+}
+
+void ASlateIMInGameWidgetBase::Server_Destroy_Implementation()
+{
+	Destroy();
+}
\ No newline at end of file
diff --git a/Engine/Plugins/Experimental/SlateIM/Source/SlateIMInGame/Private/SlateIMInGameWidgetCheatManager.cpp b/Engine/Plugins/Experimental/SlateIM/Source/SlateIMInGame/Private/SlateIMInGameWidgetCheatManager.cpp
new file mode 100644
--- /dev/null
+++ b/Engine/Plugins/Experimental/SlateIM/Source/SlateIMInGame/Private/SlateIMInGameWidgetCheatManager.cpp
@@ -0,0 +1,70 @@
+﻿// Copyright Epic Games, Inc. All Rights Reserved.
+
+#include "SlateIMInGameWidgetCheatManager.h"
+
+#include "Engine/World.h"
+#include "GameFramework/PlayerController.h"
+#include "SlateIMInGameWidgetBase.h"
+#include "SlateIMInGameWidgetModularFeature.h"
+
+USlateIMInGameWidgetCheatManager::USlateIMInGameWidgetCheatManager()
+{
+#if UE_WITH_CHEAT_MANAGER
+	if (HasAnyFlags(RF_ClassDefaultObject))
+	{
+		UCheatManager::RegisterForOnCheatManagerCreated(FOnCheatManagerCreated::FDelegate::CreateLambda(
+			[](UCheatManager* CheatManager)
+			{
+				CheatManager->AddCheatManagerExtension(NewObject<ThisClass>(CheatManager));
+			}));
+	}
+#endif
+}
+
+void USlateIMInGameWidgetCheatManager::ToggleSlateIMInGameWidget(const FString& Path) const
+{
+#if UE_WITH_CHEAT_MANAGER
+	APlayerController* PlayerController = GetPlayerController();
+	if (!PlayerController)
+	{
+		return;
+	}
+
+	const TSoftClassPtr<ASlateIMInGameWidgetBase> InGameWidgetClass = FSlateIMInGameWidgetModularFeature::FindWidgetClass(SlateIMInGameWidget::ModularFeatureName, Path);
+	if (InGameWidgetClass.Get() == nullptr)
+	{
+		UE_LOG(LogConsoleResponse, Warning, TEXT("No in game widget with path '%s'"), *Path);
+		return;
+	}
+
+	const bool bIsMenuEnabled = ASlateIMInGameWidgetBase::GetInGameWidget(PlayerController, InGameWidgetClass.Get()) != nullptr;
+	// If we're not Authority, we should be sending this cheat command to the server to execute
+	if (PlayerController->GetLocalRole() != ENetRole::ROLE_Authority)
+	{
+		PlayerController->ServerExec(FString::Printf(TEXT("EnableInGameWidgetFromClass %s %d"), *InGameWidgetClass.ToString(), bIsMenuEnabled ? 0 : 1));
+		return;
+	}
+
+	EnableInGameWidgetFromClass(InGameWidgetClass.ToString(), !bIsMenuEnabled);
+#endif
+}
+
+void USlateIMInGameWidgetCheatManager::EnableInGameWidgetFromClass(const FString& ClassPath, const bool bEnable) const
+{
+#if UE_WITH_CHEAT_MANAGER
+	APlayerController* PlayerController = GetPlayerController();
+	if (!PlayerController)
+	{
+		return;
+	}
+
+	const TSoftClassPtr<ASlateIMInGameWidgetBase> WidgetClass(ClassPath);
+	UClass* MenuClass = WidgetClass.LoadSynchronous();
+	if (!MenuClass)
+	{
+		return;
+	}
+
+	ASlateIMInGameWidgetBase::EnableInGameWidget(PlayerController, bEnable, MenuClass);
+#endif
+}
\ No newline at end of file
diff --git a/Engine/Plugins/Experimental/SlateIM/Source/SlateIMInGame/Private/SlateIMInGameWidgetCheatManager.h b/Engine/Plugins/Experimental/SlateIM/Source/SlateIMInGame/Private/SlateIMInGameWidgetCheatManager.h
new file mode 100644
--- /dev/null
+++ b/Engine/Plugins/Experimental/SlateIM/Source/SlateIMInGame/Private/SlateIMInGameWidgetCheatManager.h
@@ -0,0 +1,23 @@
+﻿// Copyright Epic Games, Inc. All Rights Reserved.
+
+#pragma once
+
+#include "GameFramework/CheatManager.h"
+
+#include "SlateIMInGameWidgetCheatManager.generated.h"
+
+UCLASS()
+class USlateIMInGameWidgetCheatManager : public UCheatManagerExtension
+{
+	GENERATED_BODY()
+
+public:
+	USlateIMInGameWidgetCheatManager();
+
+	UFUNCTION(Exec)
+	void ToggleSlateIMInGameWidget(const FString& Path) const;
+
+private:
+	UFUNCTION(Exec)
+	void EnableInGameWidgetFromClass(const FString& ClassPath, const bool bEnable) const;
+};
\ No newline at end of file
diff --git a/Engine/Plugins/Experimental/SlateIM/Source/SlateIMInGame/Private/SlateIMInGameWidgetModularFeature.cpp b/Engine/Plugins/Experimental/SlateIM/Source/SlateIMInGame/Private/SlateIMInGameWidgetModularFeature.cpp
new file mode 100644
--- /dev/null
+++ b/Engine/Plugins/Experimental/SlateIM/Source/SlateIMInGame/Private/SlateIMInGameWidgetModularFeature.cpp
@@ -0,0 +1,58 @@
+﻿// Copyright Epic Games, Inc. All Rights Reserved.
+
+#include "SlateIMInGameWidgetModularFeature.h"
+
+#include "Engine/LocalPlayer.h"
+#include "Engine/World.h"
+#include "Features/IModularFeatures.h"
+#include "GameFramework/CheatManager.h"
+#include "SlateIMInGameWidgetBase.h"
+#include "SlateIMInGameWidgetCheatManager.h"
+
+const FLazyName SlateIMInGameWidget::ModularFeatureName = TEXT("SlateIMInGameWidgets");
+
+FSlateIMInGameWidgetModularFeature::FSlateIMInGameWidgetModularFeature(const FString& InPath, const TSubclassOf<ASlateIMInGameWidgetBase>& InWidgetClass)
+	: Path(InPath)
+	, WidgetClass(InWidgetClass.Get())
+	, WidgetCommand(*FString::Printf(TEXT("SlateIM.ToggleInGameWidget.%s"), *InPath), TEXT(""), FConsoleCommandWithWorldDelegate::CreateRaw(this, &FSlateIMInGameWidgetModularFeature::ToggleWidget))
+{
+}
+
+TSoftClassPtr<ASlateIMInGameWidgetBase> FSlateIMInGameWidgetModularFeature::FindWidgetClass(const FName Type, const FString& InPath)
+{
+	TArray<FSlateIMInGameWidgetModularFeature*> DebugWidgetModularFeatures = IModularFeatures::Get().GetModularFeatureImplementations<FSlateIMInGameWidgetModularFeature>(Type);
+	for (FSlateIMInGameWidgetModularFeature* DebugWidgetModularFeature : DebugWidgetModularFeatures)
+	{
+		if (!DebugWidgetModularFeature)
+		{
+			continue;
+		}
+
+		if (!InPath.Equals(DebugWidgetModularFeature->Path, ESearchCase::IgnoreCase))
+		{
+			continue;
+		}
+
+		return DebugWidgetModularFeature->WidgetClass.Get();
+	}
+
+	return {};
+}
+
+void FSlateIMInGameWidgetModularFeature::ToggleWidget(UWorld* World)
+{
+	for (FConstPlayerControllerIterator Iterator=World->GetPlayerControllerIterator(); Iterator; ++Iterator)
+	{
+		if (APlayerController* PlayerController = Iterator->Get())
+		{
+			if (PlayerController->IsLocalController())
+			{
+				if (USlateIMInGameWidgetCheatManager* WidgetCheatManager = PlayerController->CheatManager->FindCheatManagerExtension<USlateIMInGameWidgetCheatManager>())
+				{
+					WidgetCheatManager->ToggleSlateIMInGameWidget(Path);
+					return;
+				}
+			}
+		}
+	}
+}
diff --git a/Engine/Plugins/Experimental/SlateIM/Source/SlateIMInGame/Private/SlateIMInGameWindow.cpp b/Engine/Plugins/Experimental/SlateIM/Source/SlateIMInGame/Private/SlateIMInGameWindow.cpp
new file mode 100644
--- /dev/null
+++ b/Engine/Plugins/Experimental/SlateIM/Source/SlateIMInGame/Private/SlateIMInGameWindow.cpp
@@ -0,0 +1,30 @@
+﻿// Copyright Epic Games, Inc. All Rights Reserved.
+
+#include "SlateIMInGameWindow.h"
+
+#include "SlateIM.h"
+
+ASlateIMInGameWindow::ASlateIMInGameWindow(const FObjectInitializer& ObjectInitializer)
+{
+}
+
+ASlateIMInGameWindow::ASlateIMInGameWindow(const FName& InWindowName, const FStringView& InWindowTitle)
+	: WindowName(InWindowName)
+	, WindowTitle(InWindowTitle)
+{
+}
+
+void ASlateIMInGameWindow::DrawWidget(const float DeltaTime)
+{
+	const bool bIsDrawingWindow = SlateIM::BeginWindowRoot(WindowName, WindowTitle, WindowSize);
+	if (bIsDrawingWindow)
+	{
+		DrawContent(DeltaTime);
+	}
+	SlateIM::EndRoot();
+
+	if (!bIsDrawingWindow)
+	{
+		Server_Destroy();
+	}
+}
\ No newline at end of file
diff --git a/Engine/Plugins/Experimental/SlateIM/Source/SlateIMInGame/Public/SlateIMInGameWidgetBase.h b/Engine/Plugins/Experimental/SlateIM/Source/SlateIMInGame/Public/SlateIMInGameWidgetBase.h
new file mode 100644
--- /dev/null
+++ b/Engine/Plugins/Experimental/SlateIM/Source/SlateIMInGame/Public/SlateIMInGameWidgetBase.h
@@ -0,0 +1,55 @@
+﻿// Copyright Epic Games, Inc. All Rights Reserved.
+
+#pragma once
+
+#include "GameFramework/Actor.h"
+#include "SlateIMWidgetBase.h"
+
+#include "SlateIMInGameWidgetBase.generated.h"
+
+#define UE_API SLATEIMINGAME_API
+
+class APlayerController;
+
+UCLASS(MinimalAPI)
+class ASlateIMInGameWidgetBase : public AActor
+{
+	GENERATED_BODY()
+
+public:
+	UE_API ASlateIMInGameWidgetBase();
+
+	UE_API virtual void BeginPlay() override;
+	UE_API virtual void EndPlay(EEndPlayReason::Type Reason) override;
+
+	UE_API virtual void OnRep_Owner() override;
+	UE_API virtual void PreReplication(IRepChangedPropertyTracker& ChangedPropertyTracker) override;
+	UE_API virtual bool IsNetRelevantFor(const AActor* RealViewer, const AActor* ViewTarget, const FVector& SrcLocation) const override;
+
+	APlayerController* GetPlayerController() const;
+	bool IsLocallyControlled() const;
+
+	static UE_API ASlateIMInGameWidgetBase* GetInGameWidget(const APlayerController* Owner, const TSubclassOf<ASlateIMInGameWidgetBase>& InGameWidgetClass);
+	static UE_API void EnableInGameWidget(APlayerController* Owner, const bool bEnable, const TSubclassOf<ASlateIMInGameWidgetBase>& InGameWidgetClass);
+
+protected:
+#if WITH_SERVER_CODE
+	UE_API virtual void GenerateServerSnapshot() {}
+#endif
+
+	UE_API virtual void DrawWidget(const float DeltaTime) {}
+
+	UE_API virtual void StartWidget();
+	UE_API virtual void StopWidget();
+	void TickWidget(const float DeltaTime);
+
+	static ASlateIMInGameWidgetBase* GetOrOpenInGameWidget(APlayerController* Owner, const TSubclassOf<ASlateIMInGameWidgetBase>& InGameWidgetClass);
+	static void DestroyInGameWidget(const APlayerController* Owner, const TSubclassOf<ASlateIMInGameWidgetBase>& InGameWidgetClass);
+
+	FDelegateHandle WidgetTickHandle;
+
+	UFUNCTION(Reliable, Server)
+	UE_API void Server_Destroy();
+};
+
+#undef UE_API
\ No newline at end of file
diff --git a/Engine/Plugins/Experimental/SlateIM/Source/SlateIMInGame/Public/SlateIMInGameWidgetModularFeature.h b/Engine/Plugins/Experimental/SlateIM/Source/SlateIMInGame/Public/SlateIMInGameWidgetModularFeature.h
new file mode 100644
--- /dev/null
+++ b/Engine/Plugins/Experimental/SlateIM/Source/SlateIMInGame/Public/SlateIMInGameWidgetModularFeature.h
@@ -0,0 +1,34 @@
+﻿// Copyright Epic Games, Inc. All Rights Reserved.
+
+#pragma once
+
+#include "Features/IModularFeature.h"
+#include "HAL/IConsoleManager.h"
+#include "Templates/SubclassOf.h"
+#include "UObject/SoftObjectPtr.h"
+
+#define UE_API SLATEIMINGAME_API
+
+class ASlateIMInGameWidgetBase;
+
+namespace SlateIMInGameWidget
+{
+	extern UE_API const FLazyName ModularFeatureName;
+}
+
+class FSlateIMInGameWidgetModularFeature : public IModularFeature
+{
+public:
+	UE_API FSlateIMInGameWidgetModularFeature(const FString& InPath, const TSubclassOf<ASlateIMInGameWidgetBase>& InWidgetClass);
+
+	static TSoftClassPtr<ASlateIMInGameWidgetBase> FindWidgetClass(const FName Type, const FString& InPath);
+
+protected:
+	void ToggleWidget(UWorld* World);
+
+	FString Path;
+	TSoftClassPtr<ASlateIMInGameWidgetBase> WidgetClass;
+	FAutoConsoleCommandWithWorld WidgetCommand;
+};
+
+#undef UE_API
\ No newline at end of file
diff --git a/Engine/Plugins/Experimental/SlateIM/Source/SlateIMInGame/Public/SlateIMInGameWindow.h b/Engine/Plugins/Experimental/SlateIM/Source/SlateIMInGame/Public/SlateIMInGameWindow.h
new file mode 100644
--- /dev/null
+++ b/Engine/Plugins/Experimental/SlateIM/Source/SlateIMInGame/Public/SlateIMInGameWindow.h
@@ -0,0 +1,29 @@
+﻿// Copyright Epic Games, Inc. All Rights Reserved.
+
+#pragma once
+
+#include "SlateIMInGameWidgetBase.h"
+
+#include "SlateIMInGameWindow.generated.h"
+
+#define UE_API SLATEIMINGAME_API
+
+UCLASS(MinimalAPI)
+class ASlateIMInGameWindow : public ASlateIMInGameWidgetBase
+{
+	GENERATED_BODY()
+public:
+	ASlateIMInGameWindow(const FObjectInitializer& ObjectInitializer);
+	UE_API ASlateIMInGameWindow(const FName& InWindowName, const FStringView& InWindowTitle);
+
+protected:
+	UE_API virtual void DrawWidget(const float DeltaTime) override;
+
+	UE_API virtual void DrawContent(const float DeltaTime) {};
+
+	FName WindowName;
+	FString WindowTitle;
+	FVector2f WindowSize = FVector2f(500,500);
+};
+
+#undef UE_API
\ No newline at end of file
diff --git a/Engine/Plugins/Experimental/SlateIM/Source/SlateIMInGame/SlateIMInGame.Build.cs b/Engine/Plugins/Experimental/SlateIM/Source/SlateIMInGame/SlateIMInGame.Build.cs
new file mode 100644
--- /dev/null
+++ b/Engine/Plugins/Experimental/SlateIM/Source/SlateIMInGame/SlateIMInGame.Build.cs
@@ -0,0 +1,22 @@
+﻿// Copyright Epic Games, Inc. All Rights Reserved.
+
+using UnrealBuildTool;
+
+public class SlateIMInGame : ModuleRules
+{
+	public SlateIMInGame(ReadOnlyTargetRules Target) : base(Target)
+	{
+		PCHUsage = ModuleRules.PCHUsageMode.UseExplicitOrSharedPCHs;
+
+		PublicDependencyModuleNames.AddRange(
+			new string[]
+			{
+				"Core",
+				"CoreUObject",
+				"Engine",
+				"Slate",
+				"SlateIM",
+			}
+		);
+	}
+}
\ No newline at end of file
diff --git a/Engine/Plugins/Experimental/SlateIM/Source/SlateIMInGame/SlateIMInGameModule.cpp b/Engine/Plugins/Experimental/SlateIM/Source/SlateIMInGame/SlateIMInGameModule.cpp
new file mode 100644
--- /dev/null
+++ b/Engine/Plugins/Experimental/SlateIM/Source/SlateIMInGame/SlateIMInGameModule.cpp
@@ -0,0 +1,13 @@
+﻿// Copyright Epic Games, Inc. All Rights Reserved.
+
+#include "SlateIMInGameModule.h"
+
+void FSlateIMInGameModule::StartupModule()
+{
+}
+
+void FSlateIMInGameModule::ShutdownModule()
+{
+}
+
+IMPLEMENT_MODULE(FSlateIMInGameModule, SlateIMInGame)
\ No newline at end of file
diff --git a/Engine/Plugins/Experimental/SlateIM/Source/SlateIMInGame/SlateIMInGameModule.h b/Engine/Plugins/Experimental/SlateIM/Source/SlateIMInGame/SlateIMInGameModule.h
new file mode 100644
--- /dev/null
+++ b/Engine/Plugins/Experimental/SlateIM/Source/SlateIMInGame/SlateIMInGameModule.h
@@ -0,0 +1,12 @@
+﻿// Copyright Epic Games, Inc. All Rights Reserved.
+
+#pragma once
+
+#include "Modules/ModuleManager.h"
+
+class FSlateIMInGameModule : public IModuleInterface
+{
+public:
+	virtual void StartupModule() override;
+	virtual void ShutdownModule() override;
+};
