Fixes for SlateImInGame widgets :
- make sure to return the newly created widgets
- prevent spamming server on window destruction

diff --git a/Engine/Plugins/Experimental/SlateIM/Source/SlateIMInGame/Private/SlateIMInGameWidgetBase.cpp b/Engine/Plugins/Experimental/SlateIM/Source/SlateIMInGame/Private/SlateIMInGameWidgetBase.cpp
--- a/Engine/Plugins/Experimental/SlateIM/Source/SlateIMInGame/Private/SlateIMInGameWidgetBase.cpp
+++ b/Engine/Plugins/Experimental/SlateIM/Source/SlateIMInGame/Private/SlateIMInGameWidgetBase.cpp
@@ -163,6 +163,7 @@ ASlateIMInGameWidgetBase* ASlateIMInGameWidgetBase::GetOrOpenInGameWidget(APlaye
 			if (ASlateIMInGameWidgetBase* NewWidget = Owner->GetWorld()->SpawnActor<ASlateIMInGameWidgetBase>(InGameWidgetClass, Params))
 			{
 				NewWidget->AttachToActor(Owner, FAttachmentTransformRules::KeepRelativeTransform);
+				return NewWidget;
 			}
 		}
 	}
diff --git a/Engine/Plugins/Experimental/SlateIM/Source/SlateIMInGame/Private/SlateIMInGameWidgetModularFeature.cpp b/Engine/Plugins/Experimental/SlateIM/Source/SlateIMInGame/Private/SlateIMInGameWidgetModularFeature.cpp
--- a/Engine/Plugins/Experimental/SlateIM/Source/SlateIMInGame/Private/SlateIMInGameWidgetModularFeature.cpp
+++ b/Engine/Plugins/Experimental/SlateIM/Source/SlateIMInGame/Private/SlateIMInGameWidgetModularFeature.cpp
@@ -33,7 +33,7 @@ TSoftClassPtr<ASlateIMInGameWidgetBase> FSlateIMInGameWidgetModularFeature::Find
 			continue;
 		}
 
-		return DebugWidgetModularFeature->WidgetClass.Get();
+		return DebugWidgetModularFeature->WidgetClass;
 	}
 
 	return {};
diff --git a/Engine/Plugins/Experimental/SlateIM/Source/SlateIMInGame/Private/SlateIMInGameWindow.cpp b/Engine/Plugins/Experimental/SlateIM/Source/SlateIMInGame/Private/SlateIMInGameWindow.cpp
--- a/Engine/Plugins/Experimental/SlateIM/Source/SlateIMInGame/Private/SlateIMInGameWindow.cpp
+++ b/Engine/Plugins/Experimental/SlateIM/Source/SlateIMInGame/Private/SlateIMInGameWindow.cpp
@@ -23,8 +23,9 @@ void ASlateIMInGameWindow::DrawWidget(const float DeltaTime)
 	}
 	SlateIM::EndRoot();
 
-	if (!bIsDrawingWindow)
+	if (!bIsDrawingWindow && !bDestroyRequested)
 	{
+		bDestroyRequested = true;
 		Server_Destroy();
 	}
 }
\ No newline at end of file
diff --git a/Engine/Plugins/Experimental/SlateIM/Source/SlateIMInGame/Public/SlateIMInGameWindow.h b/Engine/Plugins/Experimental/SlateIM/Source/SlateIMInGame/Public/SlateIMInGameWindow.h
--- a/Engine/Plugins/Experimental/SlateIM/Source/SlateIMInGame/Public/SlateIMInGameWindow.h
+++ b/Engine/Plugins/Experimental/SlateIM/Source/SlateIMInGame/Public/SlateIMInGameWindow.h
@@ -24,6 +24,7 @@ protected:
 	FName WindowName;
 	FString WindowTitle;
 	FVector2f WindowSize = FVector2f(500,500);
+	bool bDestroyRequested = false;
 };
 
 #undef UE_API
\ No newline at end of file
