Add root viewport scale to SlateIM

diff --git a/Engine/Plugins/Experimental/SlateIM/Source/SlateIM/Private/Roots/SlateIMViewportRoot.cpp b/Engine/Plugins/Experimental/SlateIM/Source/SlateIM/Private/Roots/SlateIMViewportRoot.cpp
--- a/Engine/Plugins/Experimental/SlateIM/Source/SlateIM/Private/Roots/SlateIMViewportRoot.cpp
+++ b/Engine/Plugins/Experimental/SlateIM/Source/SlateIM/Private/Roots/SlateIMViewportRoot.cpp
@@ -11,6 +11,7 @@
 #include "SlateIM.h"
 #include "Engine/GameViewportClient.h"
 #include "Engine/LocalPlayer.h"
+#include "Widgets/Layout/SScaleBox.h"
 #include "Widgets/SImWrapper.h"
 
 FSlateIMViewportRoot::FSlateIMViewportRoot(UGameViewportClient* InGameViewport)
@@ -86,7 +87,12 @@ void FSlateIMViewportRoot::UpdateViewport(const SlateIM::FViewportRootLayout& La
 			.Alignment(Layout.Alignment)
 			.Expose(Slot)
 			[
-				SAssignNew(Wrapper, SImWrapper)	
+				SAssignNew(ScaleBox, SScaleBox)
+				.Stretch(EStretch::UserSpecified)
+				.UserSpecifiedScale(Layout.Scale)
+				[
+					SAssignNew(Wrapper, SImWrapper)	
+				]
 			];
 		
 #if WITH_EDITOR
@@ -112,6 +118,7 @@ void FSlateIMViewportRoot::UpdateViewport(const SlateIM::FViewportRootLayout& La
 	{
 		Slot->SetAnchors(Layout.Anchors);
 		Slot->SetAlignment(Layout.Alignment);
+		ScaleBox->SetUserSpecifiedScale(Layout.Scale);
 	}
 	
 	if (Layout.Size.IsSet())
diff --git a/Engine/Plugins/Experimental/SlateIM/Source/SlateIM/Private/Roots/SlateIMViewportRoot.h b/Engine/Plugins/Experimental/SlateIM/Source/SlateIM/Private/Roots/SlateIMViewportRoot.h
--- a/Engine/Plugins/Experimental/SlateIM/Source/SlateIM/Private/Roots/SlateIMViewportRoot.h
+++ b/Engine/Plugins/Experimental/SlateIM/Source/SlateIM/Private/Roots/SlateIMViewportRoot.h
@@ -8,6 +8,7 @@
 #include "Widgets/Layout/SConstraintCanvas.h"
 
 class SImWrapper;
+class SScaleBox;
 
 namespace SlateIM
 {
@@ -49,6 +50,7 @@ public:
 #endif
 
 private:
+	TSharedPtr<SScaleBox> ScaleBox;
 	TSharedPtr<SImWrapper> Wrapper;
 	TSharedPtr<SWidget> ViewportRoot;
 	SConstraintCanvas::FSlot* Slot = nullptr;
diff --git a/Engine/Plugins/Experimental/SlateIM/Source/SlateIM/Public/SlateIM.h b/Engine/Plugins/Experimental/SlateIM/Source/SlateIM/Public/SlateIM.h
--- a/Engine/Plugins/Experimental/SlateIM/Source/SlateIM/Public/SlateIM.h
+++ b/Engine/Plugins/Experimental/SlateIM/Source/SlateIM/Public/SlateIM.h
@@ -69,6 +69,9 @@ namespace SlateIM
 
 		/** The ZOrder of the SlateIM widget in the viewport */
 		int32 ZOrder = 10000;
+
+		/** Global scale to apply within the viewport */
+		float Scale = 1.f;
 	};
 #endif // WITH_EDITOR || WITH_ENGINE
 
