SlateIM: expose GetPlayerController to other modules

diff --git a/Engine/Plugins/Experimental/SlateIM/Source/SlateIMInGame/Public/SlateIMInGameWidgetBase.h b/Engine/Plugins/Experimental/SlateIM/Source/SlateIMInGame/Public/SlateIMInGameWidgetBase.h
--- a/Engine/Plugins/Experimental/SlateIM/Source/SlateIMInGame/Public/SlateIMInGameWidgetBase.h
+++ b/Engine/Plugins/Experimental/SlateIM/Source/SlateIMInGame/Public/SlateIMInGameWidgetBase.h
@@ -26,7 +26,7 @@ public:
 	UE_API virtual void PreReplication(IRepChangedPropertyTracker& ChangedPropertyTracker) override;
 	UE_API virtual bool IsNetRelevantFor(const AActor* RealViewer, const AActor* ViewTarget, const FVector& SrcLocation) const override;
 
-	APlayerController* GetPlayerController() const;
+	UE_API APlayerController* GetPlayerController() const;
 	bool IsLocallyControlled() const;
 
 	static UE_API ASlateIMInGameWidgetBase* GetInGameWidget(const APlayerController* Owner, const TSubclassOf<ASlateIMInGameWidgetBase>& InGameWidgetClass);
