Adding SearchableComboBox in SlateIM

diff --git a/Engine/Plugins/Experimental/SlateIM/Source/SlateIM/Private/Implementation/SlateIM_Widgets.cpp b/Engine/Plugins/Experimental/SlateIM/Source/SlateIM/Private/Implementation/SlateIM_Widgets.cpp
--- a/Engine/Plugins/Experimental/SlateIM/Source/SlateIM/Private/Implementation/SlateIM_Widgets.cpp
+++ b/Engine/Plugins/Experimental/SlateIM/Source/SlateIM/Private/Implementation/SlateIM_Widgets.cpp
@@ -1,4 +1,4 @@
-﻿// Copyright Epic Games, Inc. All Rights Reserved.
+// Copyright Epic Games, Inc. All Rights Reserved.
 
 #include "SlateIM.h"
 
@@ -6,6 +6,7 @@
 #include "Misc/SlateIMManager.h"
 #include "Misc/SlateIMSlotData.h"
 #include "Misc/SlateIMWidgetScope.h"
+#include "SSearchableComboBox.h"
 #include "Widgets/Images/SImage.h"
 #include "Widgets/Input/SEditableTextBox.h"
 #include "Widgets/Input/SSlider.h"
@@ -567,6 +568,96 @@ namespace SlateIM
 		}
 	}
 
+	bool SearchableComboBox(const TArray<FString>& ComboItems, int32& InOutSelectedItemIndex, bool bForceRefresh, const FComboBoxStyle* InComboStyle)
+	{
+		FWidgetScope<SSearchableComboBox> Scope;
+		TSharedPtr<SSearchableComboBox> ComboWidget = Scope.GetWidget();
+		
+		using FComboBoxData = FSlateIMDataStore<TArray<TSharedPtr<FString>>>;
+		
+		auto BuildComboOptions = [&]()
+		{
+			TArray<TSharedPtr<FString>> Options;
+			Options.Reserve(ComboItems.Num());
+
+			for (const FString& Item : ComboItems)
+			{
+				Options.Add(MakeShared<FString>(Item));
+			}
+			return Options;
+		};
+
+		bool bValueChanged = false;
+		if (!ComboWidget)
+		{
+
+			TSharedRef<STextBlock> ComboText = SNew(STextBlock);
+			TSharedRef<FComboBoxData> WidgetIMData = MakeShared<FComboBoxData>(BuildComboOptions());
+
+			SAssignNew(ComboWidget, SSearchableComboBox)
+			.OptionsSource(&WidgetIMData->Data)
+			.InitiallySelectedItem(WidgetIMData->Data.IsValidIndex(InOutSelectedItemIndex) ? WidgetIMData->Data[InOutSelectedItemIndex] : nullptr)
+			.ComboBoxStyle(InComboStyle ? InComboStyle : &FAppStyle::Get().GetWidgetStyle<FComboBoxStyle>("ComboBox"))
+			.Content()[ComboText]
+			.OnSelectionChanged_Lambda([ActivationData = Scope.GetOrCreateActivationMetadata().ToWeakPtr()](TSharedPtr<FString> NewVal, ESelectInfo::Type SelectType)
+			{
+				FSlateIMManager::Get().ActivateWidget(ActivationData.Pin());
+			})
+			.OnGenerateWidget_Lambda([] (TSharedPtr<FString> InItem)
+			{
+				return SNew(STextBlock)
+					.Text(FText::FromString(*InItem));
+			});
+
+			const TAttribute<FText> ComboTextAttribute = TAttribute<FText>::CreateLambda(
+				[WeakComboWidget = ComboWidget->AsWeak()]()
+				{
+					if (TSharedPtr<SSearchableComboBox> PinnedComboWidget = StaticCastSharedPtr<SSearchableComboBox>(WeakComboWidget.Pin()))
+					{
+						TSharedPtr<FString> SelectedText = PinnedComboWidget->GetSelectedItem();
+						if (SelectedText.IsValid())
+						{
+							return FText::FromString(*SelectedText);
+						}
+					}
+
+					return FText::FromString(TEXT("Select..."));
+				}
+			);
+
+			ComboText->SetText(ComboTextAttribute);
+
+			ComboWidget->AddMetadata(WidgetIMData);
+
+			Scope.UpdateWidget(ComboWidget);
+		}
+		else
+		{
+			TSharedPtr<FComboBoxData> ComboBoxData = ComboWidget->GetMetaData<FComboBoxData>();
+			if (bForceRefresh)
+			{
+				ComboBoxData->Data = BuildComboOptions();
+				if (InOutSelectedItemIndex != INDEX_NONE && !ComboWidget->IsOpen() && ComboBoxData->Data.IsValidIndex(InOutSelectedItemIndex))
+				{
+					ComboWidget->SetSelectedItem(ComboBoxData->Data[InOutSelectedItemIndex]);
+				}
+			}
+			else
+			{
+				bValueChanged = Scope.IsActivatedThisFrame();
+				if (bValueChanged)
+				{
+					TSharedPtr<FString> SelectedItem = ComboWidget->GetSelectedItem();
+					const int32 NewSelectedIndex = SelectedItem.IsValid() ? ComboBoxData->Data.IndexOfByKey(SelectedItem) : INDEX_NONE;
+					UE_LOG(LogSlateIM, Verbose, TEXT("Combo Selection Changed %d -> %d"), InOutSelectedItemIndex, NewSelectedIndex);
+					InOutSelectedItemIndex = NewSelectedIndex;
+				}
+			}
+		}
+
+		return bValueChanged;
+	}
+
 	bool ComboBox(const TArray<FString>& ComboItems, int32& InOutSelectedItemIndex, bool bForceRefresh, const FComboBoxStyle* InComboStyle)
 	{
 		FWidgetScope<STextComboBox> Scope;
diff --git a/Engine/Plugins/Experimental/SlateIM/Source/SlateIM/Public/SlateIM.h b/Engine/Plugins/Experimental/SlateIM/Source/SlateIM/Public/SlateIM.h
--- a/Engine/Plugins/Experimental/SlateIM/Source/SlateIM/Public/SlateIM.h
+++ b/Engine/Plugins/Experimental/SlateIM/Source/SlateIM/Public/SlateIM.h
@@ -516,6 +516,18 @@ namespace SlateIM
 	 */
 	SLATEIM_API bool ComboBox(const TArray<FString>& ComboItems, int32& InOutSelectedItemIndex, bool bForceRefresh = false, const FComboBoxStyle* InComboStyle = nullptr);
 
+	/**
+	 * Display a searchable dropdown of text options
+	 * 
+	 * @param ComboItems The options the user can choose from
+	 * @param InOutSelectedItemIndex The index of the currently selected option
+	 * @param bForceRefresh (optional) Whether to force a refresh of the available options or the selected option (set to true for a frame when changing the list of options or manually setting the selected index)
+	 * @param InComboStyle (optional) A Combo Box Style to override the default style. Only the initial style is applied
+	 *					   
+	 * @return Whether the user changed the selected option this frame
+	 */
+	SLATEIM_API bool SearchableComboBox(const TArray<FString>& ComboItems, int32& InOutSelectedItemIndex, bool bForceRefresh = false, const FComboBoxStyle* InComboStyle = nullptr);
+
 	/**
 	 * Display a list of text options
 	 * 
diff --git a/Engine/Plugins/Experimental/SlateIM/Source/SlateIM/SlateIM.Build.cs b/Engine/Plugins/Experimental/SlateIM/Source/SlateIM/SlateIM.Build.cs
--- a/Engine/Plugins/Experimental/SlateIM/Source/SlateIM/SlateIM.Build.cs
+++ b/Engine/Plugins/Experimental/SlateIM/Source/SlateIM/SlateIM.Build.cs
@@ -19,9 +19,10 @@ public class SlateIM : ModuleRules
 			new string[]
 			{
 				"ApplicationCore",
+				"InputCore",
 				"Slate",
 				"SlateCore",
-				"InputCore",
+				"ToolWidgets",
 			}
 		);
 
