[SlateIM] Allow creation of dockable SlateIM tools

#tests Tested working properly in editor & cooked client builds

diff --git a/Engine/Plugins/Experimental/SlateIM/Source/SlateIM/Private/Misc/SlateIMExamples.cpp b/Engine/Plugins/Experimental/SlateIM/Source/SlateIM/Private/Misc/SlateIMExamples.cpp
--- a/Engine/Plugins/Experimental/SlateIM/Source/SlateIM/Private/Misc/SlateIMExamples.cpp
+++ b/Engine/Plugins/Experimental/SlateIM/Source/SlateIM/Private/Misc/SlateIMExamples.cpp
@@ -1502,6 +1502,11 @@ void FSlateIMTestWindowWidget::DrawWindow(float DeltaTime)
 	TestWidget.Draw();
 }
 
+void FSlateIMTestNomadTabWidget::DrawContent(float DeltaTime)
+{
+	TestWidget.Draw();
+}
+
 #if WITH_ENGINE
 void FSlateIMTestViewportWidget::DrawWidget(float DeltaTime)
 {
diff --git a/Engine/Plugins/Experimental/SlateIM/Source/SlateIM/Private/Misc/SlateIMExamples.h b/Engine/Plugins/Experimental/SlateIM/Source/SlateIM/Private/Misc/SlateIMExamples.h
--- a/Engine/Plugins/Experimental/SlateIM/Source/SlateIM/Private/Misc/SlateIMExamples.h
+++ b/Engine/Plugins/Experimental/SlateIM/Source/SlateIM/Private/Misc/SlateIMExamples.h
@@ -95,6 +95,19 @@ private:
 	FSlateIMTestWidget TestWidget;
 };
 
+class FSlateIMTestNomadTabWidget : public FSlateIMNomadTabBase
+{
+public:
+	FSlateIMTestNomadTabWidget(const TCHAR* Command, const TCHAR* CommandHelp)
+		: FSlateIMNomadTabBase(TEXT("SlateIM Test Suite"), Command, CommandHelp)
+	{}
+	
+	virtual void DrawContent(float DeltaTime) override;
+
+private:
+	FSlateIMTestWidget TestWidget;
+};
+
 #if WITH_ENGINE
 class FSlateIMTestViewportWidget : public FSlateIMWidgetWithCommandBase
 {
diff --git a/Engine/Plugins/Experimental/SlateIM/Source/SlateIM/Private/Misc/SlateIMManager.cpp b/Engine/Plugins/Experimental/SlateIM/Source/SlateIM/Private/Misc/SlateIMManager.cpp
--- a/Engine/Plugins/Experimental/SlateIM/Source/SlateIM/Private/Misc/SlateIMManager.cpp
+++ b/Engine/Plugins/Experimental/SlateIM/Source/SlateIM/Private/Misc/SlateIMManager.cpp
@@ -84,6 +84,7 @@ namespace SlateIM
 	FSlateIMManager::FSlateIMManager()
 #if WITH_SLATEIM_EXAMPLES
 			: TestWindowWidget(TEXT("SlateIM.ToggleTestSuiteWindow"), TEXT("Toggles the Slate immediate mode test suite window which demonstrates the capabilities of the Slate immediate mode api"))
+			, TestNomadTabWidget(TEXT("SlateIM.ToggleTestSuiteDockableTab"), TEXT("Toggles the Slate immediate mode test suite as a dockable tab which demonstrates the capabilities of the Slate immediate mode api"))
 #if WITH_ENGINE
 			, TestViewportWidget(TEXT("SlateIM.ToggleTestSuiteViewport"), TEXT("Toggles the Slate immediate mode test suite in the main viewport (PIE or LevelEditor), demonstrating the capabilities of the Slate immediate mode api"))
 #endif
diff --git a/Engine/Plugins/Experimental/SlateIM/Source/SlateIM/Private/Misc/SlateIMManager.h b/Engine/Plugins/Experimental/SlateIM/Source/SlateIM/Private/Misc/SlateIMManager.h
--- a/Engine/Plugins/Experimental/SlateIM/Source/SlateIM/Private/Misc/SlateIMManager.h
+++ b/Engine/Plugins/Experimental/SlateIM/Source/SlateIM/Private/Misc/SlateIMManager.h
@@ -216,6 +216,7 @@ namespace SlateIM
 #if WITH_SLATEIM_EXAMPLES
 		FSlateStyleBrowser SlateStyleBrowser;
 		FSlateIMTestWindowWidget TestWindowWidget;
+		FSlateIMTestNomadTabWidget TestNomadTabWidget;
 #if WITH_ENGINE
 		FSlateIMTestViewportWidget TestViewportWidget;
 #endif // WITH_ENGINE
diff --git a/Engine/Plugins/Experimental/SlateIM/Source/SlateIM/Private/SlateIMWidgetBase.cpp b/Engine/Plugins/Experimental/SlateIM/Source/SlateIM/Private/SlateIMWidgetBase.cpp
--- a/Engine/Plugins/Experimental/SlateIM/Source/SlateIM/Private/SlateIMWidgetBase.cpp
+++ b/Engine/Plugins/Experimental/SlateIM/Source/SlateIM/Private/SlateIMWidgetBase.cpp
@@ -4,7 +4,11 @@
 #include "SlateIMWidgetBase.h"
 
 #include "Framework/Application/SlateApplication.h"
+#include "Framework/Docking/SDockingTabStack.h"
+#include "Framework/Docking/TabManager.h"
 #include "SlateIM.h"
+#include "Widgets/Docking/SDockTab.h"
+#include "Widgets/SNullWidget.h"
 
 FSlateIMWidgetBase::FSlateIMWidgetBase(const FStringView& Name)
 	: WidgetName(Name)
@@ -61,6 +65,86 @@ void FSlateIMWindowBase::DrawWidget(float DeltaTime)
 	}
 }
 
+FSlateIMNomadTabBase::FSlateIMNomadTabBase(FStringView TabTitle, const TCHAR* Command, const TCHAR* CommandHelp, FSlateIcon TabIcon)
+	: FSlateIMWidgetWithCommandBase(Command, CommandHelp)
+{
+	FGlobalTabmanager::Get()->RegisterTabSpawner(GetWidgetName(), FOnSpawnTab::CreateLambda([this, TabTitleText = FText::FromStringView(TabTitle)](const FSpawnTabArgs&)
+	{
+		return SNew(SDockTab)
+			.Clipping(EWidgetClipping::ClipToBounds)
+			.TabRole(ETabRole::NomadTab)
+			.Label(TabTitleText)
+			.OnTabActivated_Lambda([this](TSharedRef<SDockTab>, ETabActivationCause) 
+				{ EnableWidget(); })
+			.OnTabClosed_Lambda([this](TSharedRef<SDockTab>) 
+				{ FSlateIMWidgetWithCommandBase::DisableWidget(); TabContent.Reset(); });
+	}))
+		.SetIcon(TabIcon);
+}
+
+FSlateIMNomadTabBase::~FSlateIMNomadTabBase()
+{
+	if (TSharedPtr<SDockTab> RootTabPin = RootTab.Pin())
+	{
+		RootTabPin->SetOnTabActivated(SDockTab::FOnTabActivatedCallback());
+		RootTabPin->SetOnTabClosed(SDockTab::FOnTabClosedCallback());
+		RootTabPin->RequestCloseTab();
+		RootTab.Reset();
+		TabContent.Reset();
+	}
+	
+	FGlobalTabmanager::Get()->UnregisterNomadTabSpawner(GetWidgetName());
+}
+
+void FSlateIMNomadTabBase::DisableWidget()
+{
+	if (TSharedPtr<SDockTab> RootTabPin = RootTab.Pin())
+	{
+		RootTabPin->RequestCloseTab();
+		RootTab.Reset();
+	}
+	
+	FSlateIMWidgetWithCommandBase::DisableWidget();
+}
+
+void FSlateIMNomadTabBase::DrawWidget(float DeltaTime)
+{
+	if (!SlateIM::CanUpdateSlateIM())
+	{
+		return;
+	}
+	
+	TSharedPtr<SDockTab> RootTabPin = RootTab.Pin();
+	if (!RootTabPin.IsValid())
+	{
+		// This is the first time we call DrawWidget() since being enabled - invoke the tab
+		RootTab = RootTabPin = FGlobalTabmanager::Get()->TryInvokeTab(GetWidgetName());
+		check(RootTabPin.IsValid())
+	}
+		
+	TSharedPtr<SWidget> NewExposedWidget;
+	if (SlateIM::BeginExposedRoot(GetWidgetName(), NewExposedWidget))
+	{
+		DrawContent(DeltaTime);
+	}
+	SlateIM::EndRoot();
+
+	if (NewExposedWidget.IsValid())
+	{
+		if (NewExposedWidget != TabContent)
+		{
+			RootTabPin->SetContent(NewExposedWidget.ToSharedRef());
+			TabContent = NewExposedWidget;
+		}
+	}
+	else
+	{
+		RootTabPin->RequestCloseTab();
+		TabContent.Reset();
+		DisableWidget();
+	}
+}
+
 TSharedRef<SWidget> FSlateIMExposedBase::GetExposedWidget() const
 {
 	return ExposedWidget.IsValid() ? ExposedWidget.ToSharedRef() : SNullWidget::NullWidget;
diff --git a/Engine/Plugins/Experimental/SlateIM/Source/SlateIM/Public/SlateIMWidgetBase.h b/Engine/Plugins/Experimental/SlateIM/Source/SlateIM/Public/SlateIMWidgetBase.h
--- a/Engine/Plugins/Experimental/SlateIM/Source/SlateIM/Public/SlateIMWidgetBase.h
+++ b/Engine/Plugins/Experimental/SlateIM/Source/SlateIM/Public/SlateIMWidgetBase.h
@@ -5,9 +5,11 @@
 #include "Delegates/IDelegateInstance.h"
 #include "HAL/IConsoleManager.h"
 #include "Math/Vector2D.h"
+#include "Textures/SlateIcon.h"
 
 #define UE_API SLATEIM_API
 
+class SDockTab;
 class SWidget;
 
 /**
@@ -25,8 +27,8 @@ public:
 	UE_API virtual ~FSlateIMWidgetBase();
 
 	void ToggleWidget() { IsWidgetEnabled() ? DisableWidget() : EnableWidget(); }
-	UE_API void EnableWidget();
-	UE_API void DisableWidget();
+	UE_API virtual void EnableWidget();
+	UE_API virtual void DisableWidget();
 
 	bool IsWidgetEnabled() const { return TickHandle.IsValid(); }
 
@@ -74,6 +76,27 @@ private:
 	FVector2f WindowSize;
 };
 
+/**
+ * Specialized version of FSlateIMWidgetBase that draws in a nomad (dockable) tab and disables itself when the user closes the window
+ */
+class FSlateIMNomadTabBase : public FSlateIMWidgetWithCommandBase
+{
+public:
+	UE_API FSlateIMNomadTabBase(FStringView TabTitle, const TCHAR* Command, const TCHAR* CommandHelp, FSlateIcon TabIcon = FSlateIcon());
+	UE_API virtual ~FSlateIMNomadTabBase();
+	
+	UE_API virtual void DisableWidget() override;
+
+protected:
+	virtual void DrawContent(float DeltaTime) = 0;
+
+private:
+	UE_API virtual void DrawWidget(float DeltaTime) final override;
+
+	TWeakPtr<SDockTab> RootTab;
+	TWeakPtr<SWidget> TabContent;
+};
+
 /**
  * Specialized version of FSlateIMWidgetBase that exposes the resulting widget for embedding in other slate widgets
  *
