[SlateIM] Fix tab registered as normal tabs rather than Nomad, and thus never unregistering

#tests Everything still working as expected in editor

diff --git a/Engine/Plugins/Experimental/SlateIM/Source/SlateIM/Private/SlateIMWidgetBase.cpp b/Engine/Plugins/Experimental/SlateIM/Source/SlateIM/Private/SlateIMWidgetBase.cpp
--- a/Engine/Plugins/Experimental/SlateIM/Source/SlateIM/Private/SlateIMWidgetBase.cpp
+++ b/Engine/Plugins/Experimental/SlateIM/Source/SlateIM/Private/SlateIMWidgetBase.cpp
@@ -68,7 +68,7 @@ void FSlateIMWindowBase::DrawWidget(float DeltaTime)
 FSlateIMNomadTabBase::FSlateIMNomadTabBase(FStringView TabTitle, const TCHAR* Command, const TCHAR* CommandHelp, FSlateIcon TabIcon)
 	: FSlateIMWidgetWithCommandBase(Command, CommandHelp)
 {
-	FGlobalTabmanager::Get()->RegisterTabSpawner(GetWidgetName(), FOnSpawnTab::CreateLambda([this, TabTitleText = FText::FromStringView(TabTitle)](const FSpawnTabArgs&)
+	FGlobalTabmanager::Get()->RegisterNomadTabSpawner(GetWidgetName(), FOnSpawnTab::CreateLambda([this, TabTitleText = FText::FromStringView(TabTitle)](const FSpawnTabArgs&)
 	{
 		return SNew(SDockTab)
 			.Clipping(EWidgetClipping::ClipToBounds)
