[SlateIM] make sure textblock widget scope data hash contains the widget pointer value to detect when the data hash index is not the same
Note: will be also applied to the other types of widget in another CL

diff --git a/Engine/Plugins/Experimental/SlateIM/Source/SlateIM/Private/Implementation/SlateIM_Widgets.cpp b/Engine/Plugins/Experimental/SlateIM/Source/SlateIM/Private/Implementation/SlateIM_Widgets.cpp
--- a/Engine/Plugins/Experimental/SlateIM/Source/SlateIM/Private/Implementation/SlateIM_Widgets.cpp
+++ b/Engine/Plugins/Experimental/SlateIM/Source/SlateIM/Private/Implementation/SlateIM_Widgets.cpp
@@ -49,28 +49,25 @@ namespace SlateIM
 	void Text(const FStringView& InText, FSlateColor Color, const FTextBlockStyle* TextStyle /*= nullptr*/)
 	{
 		FWidgetScope<STextBlock> Scope;
-		TSharedPtr<STextBlock> TextBlock = Scope.GetWidget();
+		TSharedPtr<STextBlock> TextBlock = Scope.GetOrCreate([&InText, Color, TextStyle] ()
+			{
+				TSharedPtr<STextBlock> TextBlock =
+					SNew(STextBlock)
+					.AutoWrapText(true)
+					.Text(FText::FromStringView(InText))
+					.ColorAndOpacity(Color);
+
+				if (TextStyle)
+				{
+					TextBlock->SetTextStyle(TextStyle);
+				}
+				return TextBlock;
+			});
 
 		Scope.HashData(Color);
 		Scope.HashData(TextStyle);
 		Scope.HashStringView(InText);
-
-		if (!TextBlock)
-		{
-			TextBlock =
-				SNew(STextBlock)
-				.AutoWrapText(true)
-				.Text(FText::FromStringView(InText))
-				.ColorAndOpacity(Color);
-
-			if (TextStyle)
-			{
-				TextBlock->SetTextStyle(TextStyle);
-			}
-
-			Scope.UpdateWidget(TextBlock);
-		}
-		else if(Scope.IsDataHashDirty())
+		if(Scope.IsDataHashDirty())
 		{
 			TextBlock->SetText(FText::FromStringView(InText));
 			TextBlock->SetTextStyle(TextStyle);
diff --git a/Engine/Plugins/Experimental/SlateIM/Source/SlateIM/Private/Misc/SlateIMWidgetScope.h b/Engine/Plugins/Experimental/SlateIM/Source/SlateIM/Private/Misc/SlateIMWidgetScope.h
--- a/Engine/Plugins/Experimental/SlateIM/Source/SlateIM/Private/Misc/SlateIMWidgetScope.h
+++ b/Engine/Plugins/Experimental/SlateIM/Source/SlateIM/Private/Misc/SlateIMWidgetScope.h
@@ -140,6 +140,18 @@ namespace SlateIM
 			}
 		}
 
+		TSharedPtr<WidgetType> GetOrCreate(const TFunction<TSharedPtr<WidgetType>(void)>& WidgetCreationFunction)
+		{
+			if (!Widget)
+			{
+				TSharedPtr<WidgetType> NewWidget = WidgetCreationFunction();
+				UpdateWidget(NewWidget);
+			}
+
+			HashData(PTRINT(Widget.Get()));
+			return Widget;
+		}
+
 		template<typename T>
 		void HashData(const T& Data)
 		{
