[SlateIM] change columns name identification to let user dynamically change columns between ticks.
Also fix TextBlock being reused without being modified after a column update due to wrong datahash

#Synced-CL 48202918

diff --git a/Engine/Plugins/Experimental/SlateIM/Source/SlateIM/Private/Containers/SImTableView.cpp b/Engine/Plugins/Experimental/SlateIM/Source/SlateIM/Private/Containers/SImTableView.cpp
--- a/Engine/Plugins/Experimental/SlateIM/Source/SlateIM/Private/Containers/SImTableView.cpp
+++ b/Engine/Plugins/Experimental/SlateIM/Source/SlateIM/Private/Containers/SImTableView.cpp
@@ -12,7 +12,6 @@ SLATE_IMPLEMENT_WIDGET(SImTableView)
 namespace SlateIM::Private
 {
 	const FName ExpanderColumnName = TEXT("Expander");
-	const FName BaseColumnName = TEXT("TableColumn");
 }
 
 
@@ -92,6 +91,11 @@ int32 FSlateIMTableRow::CountCellWidgetsUpToIndex(int32 Index) const
 	return CellCount;
 }
 
+int32 FSlateIMTableRow::FindColumnIndex(const FName& ColumnName) const
+{
+	return OwningTable ? OwningTable->FindColumnIndex(ColumnName) : INDEX_NONE;
+}
+
 TSharedRef<SWidget> FSlateIMTableRow::GetCellWidget(int32 CellIndex) const
 {
 	int32 CellCount = 0;
@@ -207,7 +211,7 @@ TSharedRef<SWidget> SImTableRow::GenerateWidgetForColumn(const FName& ColumnName
 	}
 	else if (TableRow)
 	{
-		const int32 ColumnIndex = ColumnName.GetNumber();
+		const int32 ColumnIndex = TableRow->FindColumnIndex(ColumnName);
 		if (ColumnIndex == 0 && TableRow->ShouldDisplayExpander())
 		{
 			return SNew(SHorizontalBox)
@@ -310,22 +314,24 @@ float SImTableView::GetNumLiveWidgets() const
 	return CachedNumLiveWidgets;
 }
 
-void SImTableView::AddColumn(const FStringView& ColumnLabel, const FStringView& ColumnToolTip, const FSlateIMSlotData& SlotData)
+void SImTableView::AddColumn(const FName& ColumnID, const FStringView& ColumnLabel, const FStringView& ColumnToolTip, const FSlateIMSlotData& SlotData)
 {
 	// Only add columns that don't already exist
 	if (Header)
 	{
 		// TODO - Handle changes to SlotData after construction
-		const int32 MaxColumnIndex = Header->GetColumns().Num();
-		if (ColumnCount >= MaxColumnIndex)
+		const int32 ColumnIdx = Header->FindColumnIndex(ColumnID);
+		if (ColumnIdx == INDEX_NONE)
 		{
-			FName ColumnName = SlateIM::Private::BaseColumnName;
-			ColumnName.SetNumber(ColumnCount);
-
-			AddColumn_Internal(ColumnName, ColumnToolTip, SlotData, ColumnLabel);
+			AddColumn_Internal(ColumnID, ColumnToolTip, SlotData, ColumnLabel);
+		}
+		else if (ColumnIdx != ColumnCount)
+		{
+			Header->SwapColumns(ColumnIdx, ColumnCount);
+			RequestRebuild();
 		}
 	}
-	
+
 	++ColumnCount;
 }
 
@@ -361,12 +367,10 @@ void SImTableView::UpdateColumns()
 {
 	if (Header)
 	{
-		FName ColumnName = SlateIM::Private::BaseColumnName;
 		const int32 MaxColumnIndex = Header->GetColumns().Num();
-		for (int32 ColumnIndex = ColumnCount; ColumnIndex < MaxColumnIndex; ++ColumnIndex)
+		if (ColumnCount < MaxColumnIndex)
 		{
-			ColumnName.SetNumber(ColumnIndex);
-			Header->RemoveColumn(ColumnName);
+			Header->RemoveColumns(ColumnCount, MaxColumnIndex - ColumnCount);
 		}
 	}
 }
@@ -415,6 +419,16 @@ int32 SImTableView::GetRowLinearizedIndex(const TSharedRef<FSlateIMTableRow>& Ro
 	return LinearizedItems.IndexOfByKey(Row);
 }
 
+bool SImTableView::IsSameColumnCount(const TSharedRef<FSlateIMTableRow>& Row) const
+{
+	return Row->GetColumnCount() == ColumnCount;
+}
+
+int32 SImTableView::FindColumnIndex(const FName& ColumnName) const
+{
+	return Header ? Header->FindColumnIndex(ColumnName) : INDEX_NONE;
+}
+
 void SImTableView::AddColumn_Internal(const FName& ColumnId, const FStringView& ColumnToolTip, const FSlateIMSlotData& SlotData, const FStringView& ColumnLabel)
 {
 	SHeaderRow::FColumn::FArguments Args = SHeaderRow::FColumn::FArguments()
diff --git a/Engine/Plugins/Experimental/SlateIM/Source/SlateIM/Private/Containers/SImTableView.h b/Engine/Plugins/Experimental/SlateIM/Source/SlateIM/Private/Containers/SImTableView.h
--- a/Engine/Plugins/Experimental/SlateIM/Source/SlateIM/Private/Containers/SImTableView.h
+++ b/Engine/Plugins/Experimental/SlateIM/Source/SlateIM/Private/Containers/SImTableView.h
@@ -27,6 +27,7 @@ public:
 	void GetChildRows(TArray<TSharedRef<FSlateIMTableRow>>& OutRows) const;
 
 	int32 CountCellWidgetsUpToIndex(int32 Index) const;
+	int32 FindColumnIndex(const FName& ColumnName) const;
 
 	TSharedRef<SWidget> GetCellWidget(int32 CellIndex) const;
 
@@ -81,7 +82,7 @@ public:
 	
 	virtual float GetNumLiveWidgets() const override;
 
-	void AddColumn(const FStringView& ColumnLabel, const FStringView& ColumnToolTip, const FSlateIMSlotData& SlotData);
+	void AddColumn(const FName& ColumnID, const FStringView& ColumnLabel, const FStringView& ColumnToolTip, const FSlateIMSlotData& SlotData);
 
 	void BeginTableUpdates();
 	void EndTableUpdates();
@@ -106,6 +107,9 @@ public:
 
 	int32 GetRowLinearizedIndex(const TSharedRef<FSlateIMTableRow>& Row) const;
 
+	bool IsSameColumnCount(const TSharedRef<FSlateIMTableRow>& Row) const;
+	int32 FindColumnIndex(const FName& ColumnName) const;
+
 private:
 	void AddColumn_Internal(const FName& ColumnId, const FStringView& ColumnToolTip, const FSlateIMSlotData& SlotData, const FStringView& ColumnLabel = FStringView());
 
diff --git a/Engine/Plugins/Experimental/SlateIM/Source/SlateIM/Private/Implementation/SlateIM_Containers.cpp b/Engine/Plugins/Experimental/SlateIM/Source/SlateIM/Private/Implementation/SlateIM_Containers.cpp
--- a/Engine/Plugins/Experimental/SlateIM/Source/SlateIM/Private/Implementation/SlateIM_Containers.cpp
+++ b/Engine/Plugins/Experimental/SlateIM/Source/SlateIM/Private/Implementation/SlateIM_Containers.cpp
@@ -239,6 +239,10 @@ namespace SlateIM
 					TableRow = MakeShared<FSlateIMTableRow>();
 					Scope.UpdateWidget(TableRow);
 				}
+				else if (Table && !Table->IsSameColumnCount(TableRow.ToSharedRef()))
+				{
+					Scope.UpdateWidget(TableRow);
+				}
 			}
 
 			FSlateIMManager::Get().PushContainer(FContainerNode(TableRow));
@@ -309,7 +313,7 @@ namespace SlateIM
 		}
 	}
 
-	void AddTableColumn(const FStringView& ColumnLabel)
+	void AddTableColumn(const FName& ColumnID, const FStringView& ColumnLabel)
 	{
 		if (TSharedPtr<SImTableView> Table = FSlateIMManager::Get().GetCurrentContainer<SImTableView>())
 		{
@@ -317,7 +321,7 @@ namespace SlateIM
 			const FSlateIMSlotData SlotData = FSlateIMManager::Get().GetCurrentAlignmentData(Defaults::Padding, HAlign_Fill, VAlign_Fill, bAutoSize
 				, Defaults::MinWidth, Defaults::MinHeight, Defaults::MaxWidth, Defaults::MaxHeight);
 			
-			Table->AddColumn(ColumnLabel, FSlateIMManager::Get().GetCurrentRoot().CurrentToolTip, SlotData);
+			Table->AddColumn(ColumnID, ColumnLabel, FSlateIMManager::Get().GetCurrentRoot().CurrentToolTip, SlotData);
 			
 			FSlateIMManager::Get().GetMutableCurrentRoot().CurrentToolTip.Empty();
 			FSlateIMManager::Get().ResetAlignmentData();
diff --git a/Engine/Plugins/Experimental/SlateIM/Source/SlateIM/Private/Misc/SlateIMExamples.cpp b/Engine/Plugins/Experimental/SlateIM/Source/SlateIM/Private/Misc/SlateIMExamples.cpp
--- a/Engine/Plugins/Experimental/SlateIM/Source/SlateIM/Private/Misc/SlateIMExamples.cpp
+++ b/Engine/Plugins/Experimental/SlateIM/Source/SlateIM/Private/Misc/SlateIMExamples.cpp
@@ -576,10 +576,10 @@ void FSlateIMTestWidget::DrawTables()
 			SlateIM::HAlign(HAlign_Fill);
 			SlateIM::MaxHeight(200.f);
 			SlateIM::BeginTable();
-			SlateIM::AddTableColumn(TEXT("Item"));
+			SlateIM::AddTableColumn(TEXT("Item"), TEXT("Item"));
 			SlateIM::HAlign(HAlign_Center);
 			SlateIM::InitialTableColumnWidth(80.f);
-			SlateIM::AddTableColumn(TEXT("Button"));
+			SlateIM::AddTableColumn(TEXT("Button"), TEXT("Button"));
 			for (int32 i = 0; i < NumTableItems; ++i)
 			{
 				if (SlateIM::NextTableCell())
@@ -614,8 +614,8 @@ void FSlateIMTestWidget::DrawTables()
 			SlateIM::HAlign(HAlign_Fill);
 			SlateIM::VAlign(VAlign_Fill);
 			SlateIM::BeginTable(nullptr, TableRowStyle);
-			SlateIM::AddTableColumn(TEXT("Name"));
-			SlateIM::AddTableColumn(TEXT("Type"));
+			SlateIM::AddTableColumn(TEXT("Name"), TEXT("Name"));
+			SlateIM::AddTableColumn(TEXT("Type"), TEXT("Type"));
 			{
 				if (SlateIM::NextTableCell())
 				{
@@ -781,7 +781,7 @@ void FSlateIMTestWidget::DrawTables()
 			SlateIM::Fill();
 			SlateIM::HAlign(HAlign_Fill);
 			SlateIM::BeginTable();
-			SlateIM::AddTableColumn(TEXT("Column"));
+			SlateIM::AddTableColumn(TEXT("Column"), TEXT("Column"));
 
 			if (SlateIM::NextTableCell())
 			{
@@ -809,6 +809,64 @@ void FSlateIMTestWidget::DrawTables()
 			}
 			SlateIM::EndTable();
 		}
+
+		// Dynamic add/remove rows and columns
+		{
+			constexpr int32 MaxColumnCount = 10;
+			static TStaticArray<bool, MaxColumnCount> Columns = {true};
+			constexpr int32 MaxRowCount = 10;
+			static TStaticArray<bool, MaxRowCount> Rows = {true};
+
+			auto GetColumnName = [] (const int32 Idx) { return FString().AppendChar('A' + Idx); };
+			auto GetRowName = [] (const int32 Idx) { return FString::FromInt(Idx); };
+
+			SlateIM::BeginHorizontalStack();
+			SlateIM::Text(TEXT("Columns:"));
+			for (int32 i = 0; i < Columns.Num(); ++i)
+			{
+				SlateIM::CheckBox(GetColumnName(i), Columns[i]);
+			}
+			SlateIM::EndHorizontalStack();
+
+			SlateIM::BeginHorizontalStack();
+			SlateIM::Text(TEXT("Rows:"));
+			for (int32 i = 0; i < Rows.Num(); ++i)
+			{
+				SlateIM::CheckBox(GetRowName(i), Rows[i]);
+			}
+			SlateIM::EndHorizontalStack();
+
+			SlateIM::MinHeight(200.f);
+			SlateIM::MaxHeight(200.f);
+			SlateIM::Fill();
+			SlateIM::HAlign(HAlign_Fill);
+			SlateIM::VAlign(VAlign_Fill);
+			SlateIM::BeginTable();
+			for (int32 i = 0; i < Columns.Num(); ++i)
+			{
+				if (Columns[i])
+				{
+					SlateIM::AddTableColumn(FName(GetColumnName(i)), GetColumnName(i));
+				}
+			}
+			for (int32 i = 0; i < Rows.Num(); ++i)
+			{
+				if (Rows[i])
+				{
+					for (int32 j = 0; j < Columns.Num(); ++j)
+					{
+						if (Columns[j])
+						{
+							if (SlateIM::NextTableCell())
+							{
+								SlateIM::Text(FString::Printf(TEXT("%s%s"), *GetColumnName(j), *GetRowName(i)));
+							}
+						}
+					}
+				}
+			}
+			SlateIM::EndTable();
+		}
 	}
 	SlateIM::EndScrollBox();
 }
@@ -1172,10 +1230,10 @@ void FSlateStyleBrowser::DrawWindow(float DeltaTime)
 	SlateIM::VAlign(VAlign_Fill);
 	SlateIM::Fill();
 	SlateIM::BeginTable();
-	SlateIM::AddTableColumn(TEXT("Name"));
-	SlateIM::AddTableColumn(TEXT("Preview"));
+	SlateIM::AddTableColumn(TEXT("Name"), TEXT("Name"));
+	SlateIM::AddTableColumn(TEXT("Preview"), TEXT("Preview"));
 	SlateIM::InitialTableColumnWidth(80.f);
-	SlateIM::AddTableColumn();
+	SlateIM::AddTableColumn(TEXT("Button"), TEXT("Button"));
 
 	const auto& Style = static_cast<const SlateIM::Private::FExposedSlateStyle&>(FAppStyle::Get());
 	const FString LowerCaseSearchString = SearchString.ToLower();
@@ -1464,8 +1522,8 @@ void FSlateStyleBrowser::DrawWindow(float DeltaTime)
 						SlateIM::VAlign(VAlign_Top);
 						SlateIM::Fill();
 						SlateIM::BeginTable(WidgetStyle.Value);
-						SlateIM::AddTableColumn(TEXT("Column 1"));
-						SlateIM::AddTableColumn(TEXT("Column 2"));
+						SlateIM::AddTableColumn(TEXT("Column 1"), TEXT("Column 1"));
+						SlateIM::AddTableColumn(TEXT("Column 2"), TEXT("Column 2"));
 						SlateIM::NextTableCell();
 						SlateIM::Text(TEXT("Cell 1"));
 						SlateIM::NextTableCell();
diff --git a/Engine/Plugins/Experimental/SlateIM/Source/SlateIM/Private/Misc/SlateIMManager.cpp b/Engine/Plugins/Experimental/SlateIM/Source/SlateIM/Private/Misc/SlateIMManager.cpp
--- a/Engine/Plugins/Experimental/SlateIM/Source/SlateIM/Private/Misc/SlateIMManager.cpp
+++ b/Engine/Plugins/Experimental/SlateIM/Source/SlateIM/Private/Misc/SlateIMManager.cpp
@@ -27,6 +27,7 @@ namespace SlateIM
 
 	void FRootNode::SetDataHash(FXxHash64 InDataHash)
 	{
+		UE_LOG(LogSlateIM, Verbose, TEXT("SetDataHash - CurrentWidgetIndex [%d] | WidgetTreeDataHash.Num() [%d]"), CurrentWidgetIndex, WidgetTreeDataHash.Num());
 		if (WidgetTreeDataHash.Num() <= CurrentWidgetIndex)
 		{
 			WidgetTreeDataHash.Emplace(FXxHash64(), InDataHash);
diff --git a/Engine/Plugins/Experimental/SlateIM/Source/SlateIM/Public/SlateIM.h b/Engine/Plugins/Experimental/SlateIM/Source/SlateIM/Public/SlateIM.h
--- a/Engine/Plugins/Experimental/SlateIM/Source/SlateIM/Public/SlateIM.h
+++ b/Engine/Plugins/Experimental/SlateIM/Source/SlateIM/Public/SlateIM.h
@@ -190,13 +190,14 @@ namespace SlateIM
 	SLATEIM_API void EndTable();
 
 	/**
-	 * Add a column to a label. Call this before adding any content to the table.
+	 * Add a column with a label. Call this before adding any content to the table.
 	 * 
+	 * @param ColumnID The name to use as ID for the column
 	 * @param ColumnLabel (optional) The text label to display in the header for this column
 	 *
 	 * @note Table columns support SetToolTip, HAlign, and VAlign. See FixedTableColumnWidth() and InitialTableColumnWidth() for size options.
 	 */
-	SLATEIM_API void AddTableColumn(const FStringView& ColumnLabel = FStringView());
+	SLATEIM_API void AddTableColumn(const FName& ColumnID, const FStringView& ColumnLabel = FStringView());
 
 	/**
 	 * Set the column to a fixed width, not resizable by users. Call this before calling AddTableColumn().
diff --git a/Engine/Plugins/Runtime/AudioMotorSim/Source/AudioMotorSimDebug/Private/AudioMotorModelDebugWidget.cpp b/Engine/Plugins/Runtime/AudioMotorSim/Source/AudioMotorSimDebug/Private/AudioMotorModelDebugWidget.cpp
--- a/Engine/Plugins/Runtime/AudioMotorSim/Source/AudioMotorSimDebug/Private/AudioMotorModelDebugWidget.cpp
+++ b/Engine/Plugins/Runtime/AudioMotorSim/Source/AudioMotorSimDebug/Private/AudioMotorModelDebugWidget.cpp
@@ -289,7 +289,7 @@ void FAudioMotorModelDebugWidget::DrawMotorSimComponentList()
 		const FTableRowStyle* TableRowStyle = &FAppStyle::Get().GetWidgetStyle<FTableRowStyle>("TableView.AlternatingRow");
 		SlateIM::BeginTable(nullptr, TableRowStyle);
 		{
-			SlateIM::AddTableColumn(TEXT("Components"));
+			SlateIM::AddTableColumn(TEXT("Components"), TEXT("Components"));
 			{
 				for(AudioMotorModelDebug::FDebugContext& DebugContext : MotorSimDebugContexts)
 				{
diff --git a/Engine/Source/Runtime/Slate/Private/Widgets/Views/SHeaderRow.cpp b/Engine/Source/Runtime/Slate/Private/Widgets/Views/SHeaderRow.cpp
--- a/Engine/Source/Runtime/Slate/Private/Widgets/Views/SHeaderRow.cpp
+++ b/Engine/Source/Runtime/Slate/Private/Widgets/Views/SHeaderRow.cpp
@@ -597,6 +597,19 @@ const TIndirectArray<SHeaderRow::FColumn>& SHeaderRow::GetColumns() const
 	return Columns;
 }
 
+int32 SHeaderRow::FindColumnIndex(const FName& InColumnId) const
+{
+	for (int32 iColumn = 0; iColumn < Columns.Num(); iColumn++ )
+	{
+		const FColumn& Column = Columns[iColumn];
+		if (Column.ColumnId == InColumnId)
+		{
+			return iColumn;
+		}
+	}
+	return INDEX_NONE;
+}
+
 void SHeaderRow::AddColumn( const FColumn::FArguments& NewColumnArgs )
 {
 	SHeaderRow::FColumn* NewColumn = new SHeaderRow::FColumn( NewColumnArgs );
@@ -648,6 +661,25 @@ void SHeaderRow::RemoveColumn( const FName& InColumnId )
 	RegenerateWidgets();
 }
 
+void SHeaderRow::RemoveColumns(const int32 StartIdx, const int32 Count)
+{
+	Columns.RemoveAt(StartIdx, Count);
+
+	ColumnsChanged.Broadcast( SharedThis( this ) );
+	RegenerateWidgets();
+}
+
+void SHeaderRow::SwapColumns(const int32 FirstColumnIdx, const int32 SecondColumnIdx)
+{
+	if (Columns.IsValidIndex(FirstColumnIdx) && Columns.IsValidIndex(SecondColumnIdx) && FirstColumnIdx != SecondColumnIdx)
+	{
+		Columns.Swap(FirstColumnIdx, SecondColumnIdx);
+
+		ColumnsChanged.Broadcast( SharedThis( this ) );
+		RegenerateWidgets();
+	}
+}
+
 void SHeaderRow::RefreshColumns()
 {
 	RegenerateWidgets();
diff --git a/Engine/Source/Runtime/Slate/Public/Widgets/Views/SHeaderRow.h b/Engine/Source/Runtime/Slate/Public/Widgets/Views/SHeaderRow.h
--- a/Engine/Source/Runtime/Slate/Public/Widgets/Views/SHeaderRow.h
+++ b/Engine/Source/Runtime/Slate/Public/Widgets/Views/SHeaderRow.h
@@ -347,6 +347,9 @@ public:
 	/** @return the Columns driven by the column headers */
 	SLATE_API const TIndirectArray<FColumn>& GetColumns() const;
 
+	/** @return the index of the column associated to the specified name */
+	SLATE_API int32 FindColumnIndex( const FName& InColumnId ) const;
+
 	/** Adds a column to the header */
 	SLATE_API void AddColumn( const FColumn::FArguments& NewColumnArgs );
 	SLATE_API void AddColumn( FColumn& NewColumn );
@@ -357,6 +360,10 @@ public:
 
 	/** Removes a column from the header */
 	SLATE_API void RemoveColumn( const FName& InColumnId );
+	SLATE_API void RemoveColumns( const int32 StartIdx, const int32 Count );
+
+	/** Swap two columns at the specified indices */
+	SLATE_API void SwapColumns( const int32 FirstColumnIdx, const int32 SecondColumnIdx );
 
 	/** Force refreshing of the column widgets*/
 	SLATE_API void RefreshColumns();
