SlateIM: Refactored table view to allow custom column title widgets (including right click menus)

#Synced-CL 48758403

diff --git a/Engine/Plugins/Experimental/SlateIM/Source/SlateIM/Private/Containers/SImTableView.cpp b/Engine/Plugins/Experimental/SlateIM/Source/SlateIM/Private/Containers/SImTableView.cpp
--- a/Engine/Plugins/Experimental/SlateIM/Source/SlateIM/Private/Containers/SImTableView.cpp
+++ b/Engine/Plugins/Experimental/SlateIM/Source/SlateIM/Private/Containers/SImTableView.cpp
@@ -5,8 +5,8 @@
 #include "Misc/SlateIMLogging.h"
 #include "Misc/SlateIMSlotData.h"
 
-SLATE_IMPLEMENT_WIDGET(SImTableRow)
 SLATE_IMPLEMENT_WIDGET(SImTableHeader)
+SLATE_IMPLEMENT_WIDGET(SImTableRow)
 SLATE_IMPLEMENT_WIDGET(SImTableView)
 
 namespace SlateIM::Private
@@ -14,6 +14,72 @@ namespace SlateIM::Private
 	const FName ExpanderColumnName = TEXT("Expander");
 }
 
+void SImTableHeader::PrivateRegisterAttributes(FSlateAttributeInitializer&)
+{
+}
+
+void FSlateIMTableHeader::RemoveUnusedChildren(int32 LastUsedChildIndex)
+{
+	const int32 NumChildren = FMath::Min(Children.Num(), LastUsedChildIndex + 1);
+
+	if (Children.Num() != NumChildren)
+	{
+		Children.SetNum(NumChildren);
+
+		if (OwningTable)
+		{
+			if (TSharedPtr<SHeaderRow> HeaderRow = OwningTable->GetHeaderRow())
+			{
+				const TIndirectArray<SHeaderRow::FColumn>& Columns = HeaderRow->GetColumns();
+
+				TArray<FName> ColumnsToRemove;
+				ColumnsToRemove.Reserve(FMath::Max(0, Columns.Num() - LastUsedChildIndex - 1));
+
+				for (int32 Index = Columns.Num() - 1; Index > LastUsedChildIndex; --Index)
+				{
+					ColumnsToRemove.Add(Columns[Index].ColumnId);
+				}
+
+				for (const FName& ColumnId : ColumnsToRemove)
+				{
+					HeaderRow->RemoveColumn(ColumnId);
+				}
+			}
+		}
+	}
+}
+
+void FSlateIMTableHeader::UpdateChild(FSlateIMChild Child, int32 Index, const FSlateIMSlotData& AlignmentData)
+{
+	if (Children.IsValidIndex(Index))
+	{
+		Children[Index] = Child;
+	}
+	else
+	{
+		Children.Add(Child);
+	}
+}
+
+TSharedPtr<SWidget> FSlateIMTableHeader::GetAsWidget()
+{
+	if (OwningTable)
+	{
+		return OwningTable->GetHeaderRow();
+	}
+
+	return nullptr;
+}
+
+TSharedPtr<SImTableView> FSlateIMTableHeader::GetOwningTable() const
+{
+	return OwningTable;
+}
+
+void FSlateIMTableHeader::SetOwningTable(const TSharedPtr<SImTableView>& InOwningTable)
+{
+	OwningTable = InOwningTable;
+}
 
 void FSlateIMTableRow::RemoveUnusedChildren(int32 LastUsedChildIndex)
 {
@@ -128,6 +194,11 @@ void FSlateIMTableRow::UpdateColumnCount(int32 NewColumnCount)
 	}
 }
 
+TSharedPtr<SImTableView> FSlateIMTableRow::GetOwningTable() const
+{
+	return OwningTable;
+}
+
 void FSlateIMTableRow::SetOwningTable(const TSharedPtr<SImTableView>& InOwningTable)
 {
 	OwningTable = InOwningTable;
@@ -239,8 +310,89 @@ TSharedRef<SWidget> SImTableRow::GenerateWidgetForColumn(const FName& ColumnName
 	return SNullWidget::NullWidget;
 }
 
-void SImTableHeader::PrivateRegisterAttributes(FSlateAttributeInitializer&)
+void FSlateIMTableBody::RemoveUnusedChildren(int32 LastUsedChildIndex)
+{
+	const int32 NumChildren = FMath::Min(Children.Num(), LastUsedChildIndex + 1);
+	if (Children.Num() != NumChildren)
+	{
+		Children.SetNum(NumChildren);
+
+		// Modifying the child count requires a rebuild to handle the new column title widgets
+		if (OwningTable)
+		{
+			OwningTable->RequestRebuild();
+		}
+	}
+}
+
+void FSlateIMTableBody::UpdateChild(FSlateIMChild Child, int32 Index, const FSlateIMSlotData& AlignmentData)
+{
+	TSharedPtr<FSlateIMTableRow> ChildRow = Child.GetChild<FSlateIMTableRow>();
+
+	if (!ensureMsgf(ChildRow, TEXT("Table bodies may only contain rows.")))
+	{
+		return;
+	}
+
+	if (OwningTable)
+	{
+		ChildRow->UpdateColumnCount(OwningTable->GetColumnCount());
+	}
+
+	ChildRow->SetOwningTable(OwningTable);
+
+	if (Children.IsValidIndex(Index))
+	{
+		Children[Index] = ChildRow.ToSharedRef();
+	}
+	else
+	{
+		Children.Add(ChildRow.ToSharedRef());
+	}
+
+	// Modifying the child count requires a rebuild to handle the case where we need to create/destroy the expander arrow and child widgets
+	if (OwningTable)
+	{
+		OwningTable->RequestRebuild();
+	}
+}
+
+TSharedPtr<SImTableView> FSlateIMTableBody::GetOwningTable() const
 {
+	return OwningTable;
+}
+
+void FSlateIMTableBody::SetOwningTable(const TSharedPtr<SImTableView>& InOwningTable)
+{
+	OwningTable = InOwningTable;
+}
+
+TSharedPtr<FSlateIMTableRow> FSlateIMTableBody::GetRow(int32 Index)
+{
+	if (Children.IsValidIndex(Index))
+	{
+		return Children[Index];
+	}
+	
+	return nullptr;
+}
+
+bool FSlateIMTableBody::IsTree() const
+{
+	for (const TSharedRef<FSlateIMTableRow>& TableRow : Children)
+	{
+		if (TableRow->HasChildRows())
+		{
+			return true;
+		}
+	}
+
+	return false;
+}
+
+const TArray<TSharedRef<FSlateIMTableRow>>& FSlateIMTableBody::GetTableRows() const
+{
+	return Children;
 }
 
 void SImTableView::PrivateRegisterAttributes(FSlateAttributeInitializer&)
@@ -249,92 +401,135 @@ void SImTableView::PrivateRegisterAttributes(FSlateAttributeInitializer&)
 
 void SImTableView::Construct(const FArguments& InArgs)
 {
+	Header = MakeShared<FSlateIMTableHeader>();
+	Header->SetOwningTable(SharedThis(this));
+
+	Body = MakeShared<FSlateIMTableBody>();
+	Body->SetOwningTable(SharedThis(this));
+
 	Super::Construct(
 		FArguments(InArgs)
-		.TreeItemsSource(&TableRows)
-		.HeaderRow(SAssignNew(Header, SImTableHeader))
+		.TreeItemsSource(&Body->GetTableRows())  // Cannot set IItemsSource here, but still needs data.
+		.HeaderRow(SNew(SImTableHeader, Header.ToSharedRef()))
 		.OnGenerateRow(this, &SImTableView::GenerateRow)
 		.OnGetChildren(this, &SImTableView::GatherChildren));
+
+	SetRootItemsSource(MakeUnique<FRowSource>(*this));
+	RequestListRefresh();
 }
 
 FSlateIMChild SImTableView::GetChild(int32 Index)
 {
-	return TableRows.IsValidIndex(Index) ? FSlateIMChild(TableRows[Index]) : FSlateIMChild(nullptr);
+	switch (Index)
+	{
+		case 0:
+			return FSlateIMChild(Header.ToSharedRef());
+
+		case 1:
+			return FSlateIMChild(Body.ToSharedRef());
+
+		default:
+			return {};
+	}
 }
 
 void SImTableView::RemoveUnusedChildren(int32 LastUsedChildIndex)
 {
-	if (TableRows.Num() > (LastUsedChildIndex + 1))
-	{
-		RequestRefresh();
-		TableRows.SetNum(LastUsedChildIndex + 1);
-	}
+	// Cannot reset the children.
 }
 
 void SImTableView::UpdateChild(FSlateIMChild Child, int32 Index, const FSlateIMSlotData& AlignmentData)
 {
-	TSharedPtr<FSlateIMTableRow> Row = Child.GetChild<FSlateIMTableRow>();
-	if (!ensureMsgf(Row, TEXT("Tables can only hold Rows/Cells")))
+	switch (Index)
 	{
-		return;
-	}
+		case 0:
+			if (TSharedPtr<FSlateIMTableHeader> ChildHeader = Child.GetChild<FSlateIMTableHeader>())
+			{
+				Header = ChildHeader;
 
-	Row->SetOwningTable(SharedThis(this));
+				if (HeaderRow)
+				{
+					HeaderRow->RefreshColumns();
+				}
+			}
+			break;
 
-	if (TableRows.IsValidIndex(Index))
-	{
-		if (TableRows[Index] != Row)
-		{
-			RequestRefresh();
-			TableRows[Index] = Row.ToSharedRef();
-		}
-	}
-	else
-	{
-		RequestRefresh();
-		TableRows.Add(Row.ToSharedRef());
-	}
+		case 1:
+			if (TSharedPtr<FSlateIMTableBody> ChildBody = Child.GetChild<FSlateIMTableBody>())
+			{
+				Body = ChildBody;
+				RequestLayoutRefresh();
+			}
+			break;
 
-	if (Row->GetColumnCount() != ColumnCount)
-	{
-		RequestRefresh();
-		Row->UpdateColumnCount(ColumnCount);
+		default:
+			checkNoEntry();
+			break;
 	}
 }
 
 float SImTableView::GetNumLiveWidgets() const
 {
-	// NumLiveWidgets can be 0 during updates, so we use the last valid value we had
-	const int32 NumLiveWidgets = STreeView<TSharedRef<FSlateIMTableRow>>::GetNumLiveWidgets();
-	if (NumLiveWidgets > 0)
+	if (ItemsPanel)
 	{
-		CachedNumLiveWidgets = NumLiveWidgets;
+		// NumLiveWidgets can be 0 during updates, so we use the last valid value we had
+		const int32 NumLiveWidgets = STreeView<TSharedRef<FSlateIMTableRow>>::GetNumLiveWidgets();
+
+		if (NumLiveWidgets > 0)
+		{
+			CachedNumLiveWidgets = NumLiveWidgets;
+		}
 	}
 
 	return CachedNumLiveWidgets;
 }
 
-void SImTableView::AddColumn(const FName& ColumnID, const FStringView& ColumnLabel, const FStringView& ColumnToolTip, const FSlateIMSlotData& SlotData)
+TSharedPtr<FSlateIMTableHeader> SImTableView::GetHeader() const
+{
+	return Header;
+}
+
+TSharedPtr<FSlateIMTableBody> SImTableView::GetBody() const
+{
+	return Body;
+}
+
+void SImTableView::AddColumn(const FName& ColumnID, const FStringView& ColumnToolTip, const FSlateIMSlotData& SlotData)
 {
 	// Only add columns that don't already exist
-	if (Header)
+	if (HeaderRow)
 	{
 		// TODO - Handle changes to SlotData after construction
-		const int32 ColumnIdx = Header->FindColumnIndex(ColumnID);
+		int32 ColumnIdx = HeaderRow->FindColumnIndex(ColumnID);
+		TSharedRef<SWidget> Content = SNullWidget::NullWidget;
+
 		if (ColumnIdx == INDEX_NONE)
 		{
-			AddColumn_Internal(ColumnID, ColumnToolTip, SlotData, ColumnLabel);
+			AddColumn_Internal(ColumnID, ColumnToolTip, SlotData);
 		}
 		else if (ColumnIdx != ColumnCount)
 		{
-			Header->SwapColumns(ColumnIdx, ColumnCount);
-			RequestRebuild();
-		}
+			HeaderRow->SwapColumns(ColumnIdx, ColumnCount);
+			HeaderRow->RefreshColumns();
+		}		
 	}
 
 	++ColumnCount;
 }
 
+void SImTableView::EndColumnUpdates()
+{
+	if (Header && HeaderRow)
+	{
+		TIndirectArray<SHeaderRow::FColumn>& Columns = const_cast<TIndirectArray<SHeaderRow::FColumn>&>(HeaderRow->GetColumns());
+
+		for (int32 Index = 0; Index < Columns.Num(); ++Index)
+		{
+			Columns[Index].HeaderContent.Widget = Header->GetChild(Index).GetWidgetRef();
+		}
+	}
+}
+
 void SImTableView::BeginTableUpdates()
 {
 	ColumnCount = 0;
@@ -346,11 +541,21 @@ void SImTableView::EndTableUpdates()
 	{
 		DirtyState = Clean;
 		RequestListRefresh();
+
+		if (HeaderRow)
+		{
+			HeaderRow->RefreshColumns();
+		}
 	}
 	else if (DirtyState == NeedsRebuild)
 	{
 		DirtyState = Clean;
 		RebuildList();
+
+		if (HeaderRow)
+		{
+			HeaderRow->RefreshColumns();
+		}
 	}
 }
 
@@ -365,12 +570,12 @@ void SImTableView::BeginTableContent()
 
 void SImTableView::UpdateColumns()
 {
-	if (Header)
+	if (HeaderRow)
 	{
-		const int32 MaxColumnIndex = Header->GetColumns().Num();
+		const int32 MaxColumnIndex = HeaderRow->GetColumns().Num();
 		if (ColumnCount < MaxColumnIndex)
 		{
-			Header->RemoveColumns(ColumnCount, MaxColumnIndex - ColumnCount);
+			HeaderRow->RemoveColumns(ColumnCount, MaxColumnIndex - ColumnCount);
 		}
 	}
 }
@@ -382,15 +587,7 @@ void SImTableView::SetTableRowStyle(const FTableRowStyle* InRowStyle)
 
 bool SImTableView::IsTree() const
 {
-	for (const TSharedRef<FSlateIMTableRow>& TableRow : TableRows)
-	{
-		if (TableRow->HasChildRows())
-		{
-			return true;
-		}
-	}
-
-	return false;
+	return Body->IsTree();
 }
 
 void SImTableView::RequestRefresh()
@@ -414,6 +611,11 @@ void SImTableView::GatherChildren(TSharedRef<FSlateIMTableRow> Row, TArray<TShar
 	Row->GetChildRows(OutChildren);
 }
 
+TSharedPtr<FSlateIMTableRow> SImTableView::GetRow(int32 Index)
+{ 
+	return Body->GetRow(Index);
+}
+
 int32 SImTableView::GetRowLinearizedIndex(const TSharedRef<FSlateIMTableRow>& Row) const
 {
 	return LinearizedItems.IndexOfByKey(Row);
@@ -426,17 +628,17 @@ bool SImTableView::IsSameColumnCount(const TSharedRef<FSlateIMTableRow>& Row) co
 
 int32 SImTableView::FindColumnIndex(const FName& ColumnName) const
 {
-	return Header ? Header->FindColumnIndex(ColumnName) : INDEX_NONE;
+	return HeaderRow ? HeaderRow->FindColumnIndex(ColumnName) : INDEX_NONE;
 }
 
-void SImTableView::AddColumn_Internal(const FName& ColumnId, const FStringView& ColumnToolTip, const FSlateIMSlotData& SlotData, const FStringView& ColumnLabel)
+void SImTableView::AddColumn_Internal(const FName& ColumnId, const FStringView& ColumnToolTip, const FSlateIMSlotData& SlotData)
 {
 	SHeaderRow::FColumn::FArguments Args = SHeaderRow::FColumn::FArguments()
 		.ColumnId(ColumnId)
-		.DefaultLabel(FText::FromStringView(ColumnLabel))
 		.DefaultTooltip(FText::FromStringView(ColumnToolTip))
-		.HAlignHeader(SlotData.HorizontalAlignment)
-		.VAlignHeader(SlotData.VerticalAlignment);
+		.HAlignHeader(HAlign_Fill)
+		.VAlignHeader(VAlign_Fill)
+		.HeaderContentPadding(SlotData.Padding);
 
 	if (SlotData.bAutoSize && SlotData.MinWidth > 0)
 	{
@@ -460,7 +662,19 @@ void SImTableView::AddColumn_Internal(const FName& ColumnId, const FStringView&
 			Args.FillWidth(1.f);
 		}
 	}
-		
-	Header->AddColumn(Args);
+	
+	HeaderRow->AddColumn(Args);
 	RequestRefresh();
 }
+
+void SImTableHeader::Construct(const SHeaderRow::FArguments& InArgs, const TSharedRef<FSlateIMTableHeader> InHeader)
+{
+	Header = InHeader;
+
+	SHeaderRow::Construct(InArgs);
+}
+
+TSharedPtr<FSlateIMTableHeader> SImTableHeader::GetHeader() const
+{
+	return Header;
+}
diff --git a/Engine/Plugins/Experimental/SlateIM/Source/SlateIM/Private/Containers/SImTableView.h b/Engine/Plugins/Experimental/SlateIM/Source/SlateIM/Private/Containers/SImTableView.h
--- a/Engine/Plugins/Experimental/SlateIM/Source/SlateIM/Private/Containers/SImTableView.h
+++ b/Engine/Plugins/Experimental/SlateIM/Source/SlateIM/Private/Containers/SImTableView.h
@@ -4,11 +4,57 @@
 
 #include "ISlateIMContainer.h"
 #include "Misc/ISlateIMChild.h"
+#include "Widgets/Views/IItemsSource.h"
 #include "Widgets/Views/STreeView.h"
 #include "Widgets/Views/STableRow.h"
 
 class SImTableView;
 
+/**
+ * Header
+ */
+class FSlateIMTableHeader : public ISlateIMContainer, public ISlateIMChild, public TSharedFromThis<FSlateIMTableHeader>
+{
+	SLATE_IM_TYPE_DATA(FSlateIMTableHeader, ISlateIMChild)
+	// Deadly diamond of death, but it's okay since ISlateIMContainer and ISlateIMChild don't have implementations for ISlateIMTypeChecking
+	// SLATE_IM_TYPE_DATA(FSlateIMTableHeader, ISlateIMContainer)
+
+public:
+	virtual FString GetDebugName() override { return GetTypeId().ToString(); }
+	virtual int32 GetNumChildren() override { return Children.Num(); }
+	virtual FSlateIMChild GetChild(int32 Index) override { return Children.IsValidIndex(Index) ? Children[Index] : nullptr; }
+	virtual FSlateIMChild GetContainer() override { return TSharedRef<ISlateIMChild>(AsShared()); }
+	virtual void RemoveUnusedChildren(int32 LastUsedChildIndex) override;
+	virtual void UpdateChild(FSlateIMChild Child, int32 Index, const FSlateIMSlotData& AlignmentData) override;
+
+	virtual TSharedPtr<SWidget> GetAsWidget() override;
+
+	TSharedPtr<SImTableView> GetOwningTable() const;
+	void SetOwningTable(const TSharedPtr<SImTableView>& InOwningTable);
+
+private:
+	TArray<FSlateIMChild> Children;
+	TSharedPtr<SImTableView> OwningTable;
+};
+
+class SImTableHeader : public SHeaderRow
+{
+	SLATE_DECLARE_WIDGET(SImTableHeader, SHeaderRow)
+
+public:
+	void Construct(const SHeaderRow::FArguments& InArgs, const TSharedRef<FSlateIMTableHeader> InHeader);
+
+	TSharedPtr<FSlateIMTableHeader> GetHeader() const;
+
+private:
+	TSharedPtr<FSlateIMTableHeader> Header;
+};
+
+
+
+/**
+ * Row
+ */
 class FSlateIMTableRow : public ISlateIMContainer, public ISlateIMChild, public TSharedFromThis<FSlateIMTableRow>
 {
 	SLATE_IM_TYPE_DATA(FSlateIMTableRow, ISlateIMChild)
@@ -34,6 +80,7 @@ public:
 	int32 GetColumnCount() const { return ColumnCount; }
 	void UpdateColumnCount(int32 NewColumnCount);
 
+	TSharedPtr<SImTableView> GetOwningTable() const;
 	void SetOwningTable(const TSharedPtr<SImTableView>& InOwningTable);
 
 	bool IsExpanded();
@@ -61,12 +108,48 @@ private:
 	TSharedPtr<FSlateIMTableRow> TableRow;
 };
 
-class SImTableHeader : public SHeaderRow
+
+
+/**
+ * Body
+ */
+class FSlateIMTableBody : public ISlateIMContainer, public ISlateIMChild, public TSharedFromThis<FSlateIMTableBody>
 {
-	SLATE_DECLARE_WIDGET(SImTableHeader, SHeaderRow)
+	SLATE_IM_TYPE_DATA(FSlateIMTableBody, ISlateIMChild)
+	// Deadly diamond of death, but it's okay since ISlateIMContainer and ISlateIMChild don't have implementations for ISlateIMTypeChecking
+	// SLATE_IM_TYPE_DATA(FSlateIMTableBody, ISlateIMContainer)
+
+public:
+	virtual FString GetDebugName() override { return GetTypeId().ToString(); }
+	virtual int32 GetNumChildren() override { return Children.Num(); }
+	virtual FSlateIMChild GetChild(int32 Index) override { return Children.IsValidIndex(Index) ? FSlateIMChild(Children[Index]) : nullptr; }
+	virtual FSlateIMChild GetContainer() override { return TSharedRef<ISlateIMChild>(AsShared()); }
+	virtual void RemoveUnusedChildren(int32 LastUsedChildIndex) override;
+	virtual void UpdateChild(FSlateIMChild Child, int32 Index, const FSlateIMSlotData& AlignmentData) override;
+	virtual TSharedPtr<SWidget> GetAsWidget() override { return nullptr; }
+
+	TSharedPtr<SImTableView> GetOwningTable() const;
+	void SetOwningTable(const TSharedPtr<SImTableView>& InOwningTable);
+
+	TSharedPtr<FSlateIMTableRow> GetRow(int32 Index);
+
+	bool IsTree() const;
+
+	const TArray<TSharedRef<FSlateIMTableRow>>& GetTableRows() const;
+
+private:
+	TArray<TSharedRef<FSlateIMTableRow>> Children;
+	TSharedPtr<SImTableView> OwningTable;
 };
 
-class SImTableView : public STreeView<TSharedRef<FSlateIMTableRow>>, public ISlateIMContainer
+// No direct widget representation
+
+
+
+/**
+ * Table
+ */
+class SImTableView : public ISlateIMContainer, public STreeView<TSharedRef<FSlateIMTableRow>>
 {
 	SLATE_DECLARE_WIDGET(SImTableView, STreeView<TSharedRef<FSlateIMTableRow>>)
 	SLATE_IM_TYPE_DATA(SImTableView, ISlateIMContainer)
@@ -74,7 +157,7 @@ class SImTableView : public STreeView<TSharedRef<FSlateIMTableRow>>, public ISla
 public:
 	void Construct(const FArguments& InArgs);
 
-	virtual int32 GetNumChildren() override { return TableRows.Num(); }
+	virtual int32 GetNumChildren() override { return 2; }
 	virtual FSlateIMChild GetChild(int32 Index) override;
 	virtual FSlateIMChild GetContainer() override { return AsShared(); }
 	virtual void RemoveUnusedChildren(int32 LastUsedChildIndex) override;
@@ -82,7 +165,14 @@ public:
 	
 	virtual float GetNumLiveWidgets() const override;
 
-	void AddColumn(const FName& ColumnID, const FStringView& ColumnLabel, const FStringView& ColumnToolTip, const FSlateIMSlotData& SlotData);
+	TSharedPtr<FSlateIMTableHeader> GetHeader() const;
+	TSharedPtr<FSlateIMTableBody> GetBody() const;
+
+	void AddColumn(const FName& ColumnID, const FStringView& ColumnToolTip, const FSlateIMSlotData& SlotData);
+
+	int32 GetColumnCount() const { return ColumnCount; }
+
+	void EndColumnUpdates();
 
 	void BeginTableUpdates();
 	void EndTableUpdates();
@@ -103,7 +193,7 @@ public:
 
 	void GatherChildren(TSharedRef<FSlateIMTableRow> Row, TArray<TSharedRef<FSlateIMTableRow>>& OutChildren) const;
 
-	TSharedPtr<FSlateIMTableRow> GetRow(int32 Index) { return TableRows.IsValidIndex(Index) ? TableRows[Index].ToSharedPtr() : nullptr; }
+	TSharedPtr<FSlateIMTableRow> GetRow(int32 Index);
 
 	int32 GetRowLinearizedIndex(const TSharedRef<FSlateIMTableRow>& Row) const;
 
@@ -111,10 +201,36 @@ public:
 	int32 FindColumnIndex(const FName& ColumnName) const;
 
 private:
-	void AddColumn_Internal(const FName& ColumnId, const FStringView& ColumnToolTip, const FSlateIMSlotData& SlotData, const FStringView& ColumnLabel = FStringView());
+	void AddColumn_Internal(const FName& ColumnId, const FStringView& ColumnToolTip, const FSlateIMSlotData& SlotData);
 
-	TArray<TSharedRef<FSlateIMTableRow>> TableRows;
-	TSharedPtr<SImTableHeader> Header;
+	// Because the row source is now in a different object which can change, use this intermediary object 
+	// to avoid issues updating the source.
+	struct FRowSource : public UE::Slate::ItemsSource::IItemsSource<TSharedRef<FSlateIMTableRow>>
+	{
+		SImTableView& OwningTable;
+
+		FRowSource(SImTableView& InOwningTable)
+			: OwningTable(InOwningTable)
+		{}
+
+		virtual const TArrayView<const TSharedRef<FSlateIMTableRow>> GetItems() const override
+		{
+			if (TSharedPtr<FSlateIMTableBody> BodyLocal = OwningTable.GetBody())
+			{
+				return BodyLocal->GetTableRows();
+			}
+
+			return {};
+		}
+
+		virtual bool IsSame(const void* RawPointer) const override
+		{
+			return this == RawPointer;
+		}
+	};
+
+	TSharedPtr<FSlateIMTableHeader> Header;
+	TSharedPtr<FSlateIMTableBody> Body;
 	const FTableRowStyle* RowStyle = nullptr;
 	int32 ColumnCount = 0;
 
diff --git a/Engine/Plugins/Experimental/SlateIM/Source/SlateIM/Private/Implementation/SlateIM_Containers.cpp b/Engine/Plugins/Experimental/SlateIM/Source/SlateIM/Private/Implementation/SlateIM_Containers.cpp
--- a/Engine/Plugins/Experimental/SlateIM/Source/SlateIM/Private/Implementation/SlateIM_Containers.cpp
+++ b/Engine/Plugins/Experimental/SlateIM/Source/SlateIM/Private/Implementation/SlateIM_Containers.cpp
@@ -1,4 +1,4 @@
-﻿// Copyright Epic Games, Inc. All Rights Reserved.
+// Copyright Epic Games, Inc. All Rights Reserved.
 
 #include "SlateIM.h"
 
@@ -223,11 +223,22 @@ namespace SlateIM
 
 	void BeginTableRow()
 	{
-		// If we're in a row, we're adding a child row, otherwise it's a top-level row in the table
-		TSharedPtr<SImTableView> Table = FSlateIMManager::Get().GetCurrentContainer<SImTableView>();
-		TSharedPtr<FSlateIMTableRow> ParentRow = !Table.IsValid() ? FSlateIMManager::Get().GetCurrentContainer<FSlateIMTableRow>() : nullptr;
-		if (ensureMsgf(Table || ParentRow, TEXT("Table Rows/Cells can only exist within Tables")))
+		// If we're in a row, we're adding a child row, otherwise it's a top-level row in the body
+		TSharedPtr<FSlateIMTableBody> Body = FSlateIMManager::Get().GetCurrentContainer<FSlateIMTableBody>();
+		TSharedPtr<FSlateIMTableRow> ParentRow = !Body.IsValid() ? FSlateIMManager::Get().GetCurrentContainer<FSlateIMTableRow>() : nullptr;
+		if (ensureMsgf(Body || ParentRow, TEXT("Table Rows/Cells can only exist within Table Bodies")))
 		{
+			TSharedPtr<SImTableView> Table;
+
+			if (Body)
+			{
+				Table = Body->GetOwningTable();
+			}
+			else if (ParentRow)
+			{
+				Table = ParentRow->GetOwningTable();
+			}
+
 			TSharedPtr<FSlateIMTableRow> TableRow;
 			// Note widget scope must fall out of scope before we push container or else we are pushing this new container into itself instead of into its parent
 			{
@@ -253,6 +264,39 @@ namespace SlateIM
 	{
 		FSlateIMManager::Get().PopContainer<FSlateIMTableRow>();
 	}
+
+	void BeginTableColumn(const FName& ColumnID)
+	{
+		if (TSharedPtr<FSlateIMTableHeader> Header = FSlateIMManager::Get().GetCurrentContainer<FSlateIMTableHeader>())
+		{
+			if (TSharedPtr<SImTableView> Table = Header->GetOwningTable())
+			{
+				constexpr bool bAutoSize = false;
+				const FSlateIMSlotData SlotData = FSlateIMManager::Get().GetCurrentAlignmentData(Defaults::Padding, HAlign_Fill, VAlign_Fill, bAutoSize
+					, Defaults::MinWidth, Defaults::MinHeight, Defaults::MaxWidth, Defaults::MaxHeight);
+
+				Table->AddColumn(ColumnID, FSlateIMManager::Get().GetCurrentRoot().CurrentToolTip, SlotData);
+
+				FSlateIMManager::Get().GetMutableCurrentRoot().CurrentToolTip.Empty();
+				FSlateIMManager::Get().ResetAlignmentData();
+
+				SlateIM::BeginHorizontalStack();
+			}
+			else
+			{
+				ensureAlwaysMsgf(Table, TEXT("Internal error: Table header has not been assigned its owning table."));
+			}
+		}
+		else
+		{
+			ensureAlwaysMsgf(Header, TEXT("Current container should be a table header - Is there a missing SlateIM::BeginTable() or SlateIM::BeginTableHeader() statement?"));
+		}
+	}
+
+	void EndTableColumn()
+	{
+		SlateIM::EndHorizontalStack();
+	}
 	
 	void BeginTable(const FTableViewStyle* InStyle, const FTableRowStyle* InRowStyle)
 	{
@@ -286,26 +330,70 @@ namespace SlateIM
 
 	void EndTable()
 	{
-		// If there's an open cell, close it
+		// If there's an open column or cell, close it
 		if (TSharedPtr<SImStackBox> Cell = FSlateIMManager::Get().GetCurrentContainer<SImStackBox>())
 		{
-			SlateIM::EndVerticalStack();
-		}
+			// Column
+			if (Cell->GetOrientation() == Orient_Horizontal)
+			{
+				EndTableColumn();
 
-		TSharedPtr<FSlateIMTableRow> Row = FSlateIMManager::Get().GetCurrentContainer<FSlateIMTableRow>();
-		TSharedPtr<SImTableView> Table = FSlateIMManager::Get().GetCurrentContainer<SImTableView>();
-		ensureAlwaysMsgf((Row.IsValid() || Table.IsValid()), TEXT("Current container should be a row or a table - Is there a missing SlateIM::EndX() statement?"));
-		if (Row)
+				// Pop header
+				if (FSlateIMManager::Get().GetCurrentContainer<FSlateIMTableHeader>())
+				{
+					EndTableHeader();
+				}
+			}
+			// Cell
+			else if (Cell->GetOrientation() == Orient_Vertical)
+			{
+				SlateIM::EndVerticalStack();
+
+				// Pop row
+				while (FSlateIMManager::Get().GetCurrentContainer<FSlateIMTableRow>())
+				{
+					EndTableRow();
+				}
+
+				// Pop body
+				if (FSlateIMManager::Get().GetCurrentContainer<FSlateIMTableBody>())
+				{
+					EndTableBody();
+				}
+			}
+		}
+		// Pop header
+		else if (FSlateIMManager::Get().GetCurrentContainer<FSlateIMTableHeader>())
 		{
-			EndTableRow();
+			EndTableHeader();
 		}
+		// Pop row
+		else if (FSlateIMManager::Get().GetCurrentContainer<FSlateIMTableRow>())
+		{
+			while (FSlateIMManager::Get().GetCurrentContainer<FSlateIMTableRow>())
+			{
+				EndTableRow();
+			}
 
-		// If we just ended a Row, then we need to fetch the table
-		if (!Table.IsValid())
+			// Pop body
+			if (FSlateIMManager::Get().GetCurrentContainer<FSlateIMTableBody>())
+			{
+				EndTableBody();
+			}
+		}
+		// Pop body
+		else if (FSlateIMManager::Get().GetCurrentContainer<FSlateIMTableBody>())
 		{
-			Table = FSlateIMManager::Get().GetCurrentContainer<SImTableView>();
+			EndTableBody();
 		}
-		
+		else if (!FSlateIMManager::Get().GetCurrentContainer<SImTableView>())
+		{
+			ensureAlwaysMsgf(false, TEXT("Current container should be a row, body or a table - Is there a missing SlateIM::EndX() statement?"));
+			return;
+		}
+
+		TSharedPtr<SImTableView> Table = FSlateIMManager::Get().GetCurrentContainer<SImTableView>();
+
 		if (ensureAlwaysMsgf(Table, TEXT("Current container should be a table - Is there a missing SlateIM::EndX() statement?")))
 		{
 			FSlateIMManager::Get().PopContainer<SImTableView>();
@@ -313,21 +401,76 @@ namespace SlateIM
 		}
 	}
 
-	void AddTableColumn(const FName& ColumnID, const FStringView& ColumnLabel)
+	void BeginTableHeader()
 	{
 		if (TSharedPtr<SImTableView> Table = FSlateIMManager::Get().GetCurrentContainer<SImTableView>())
 		{
-			constexpr bool bAutoSize = false;
-			const FSlateIMSlotData SlotData = FSlateIMManager::Get().GetCurrentAlignmentData(Defaults::Padding, HAlign_Fill, VAlign_Fill, bAutoSize
-				, Defaults::MinWidth, Defaults::MinHeight, Defaults::MaxWidth, Defaults::MaxHeight);
-			
-			Table->AddColumn(ColumnID, ColumnLabel, FSlateIMManager::Get().GetCurrentRoot().CurrentToolTip, SlotData);
-			
-			FSlateIMManager::Get().GetMutableCurrentRoot().CurrentToolTip.Empty();
-			FSlateIMManager::Get().ResetAlignmentData();
+			FSlateIMManager::Get().PushContainer(FContainerNode(Table->GetHeader()));
+		}
+		else
+		{
+			ensureAlwaysMsgf(false, TEXT("Current container should be a table - Is there a missing SlateIM::BeginTable() statement?"));
 		}
 	}
 
+	void EndTableHeader()
+	{
+		// If there's an open column, close it
+		if (TSharedPtr<SImStackBox> Cell = FSlateIMManager::Get().GetCurrentContainer<SImStackBox>())
+		{
+			EndTableColumn();
+		}
+
+		FSlateIMManager::Get().PopContainer<FSlateIMTableHeader>();
+
+		if (TSharedPtr<SImTableView> Table = FSlateIMManager::Get().GetCurrentContainer<SImTableView>())
+		{
+			Table->EndColumnUpdates();
+		}
+	}
+
+	void NextTableColumn(const FName& ColumnID)
+	{
+		// If we haven't started the header, do so now.
+		if (TSharedPtr<SImTableView> Table = FSlateIMManager::Get().GetCurrentContainer<SImTableView>())
+		{
+			BeginTableHeader();
+		}
+		// If there's an open column, close it
+		else if (TSharedPtr<SImStackBox> Cell = FSlateIMManager::Get().GetCurrentContainer<SImStackBox>())
+		{
+			EndTableColumn();
+		}
+		else if (!FSlateIMManager::Get().GetCurrentContainer<FSlateIMTableHeader>())
+		{
+			ensureAlwaysMsgf(false, TEXT("Current container should be a table or a table header - Is there a missing SlateIM::BeginTable() or SlateIM::BeginTableHeader() statement?"));
+			return;
+		}
+
+		BeginTableColumn(ColumnID);
+	}
+
+	void AddTableColumn(const FName& ColumnID, const FStringView& ColumnLabel)
+	{
+		FSlateIMAlignmentState AlignmentState = FSlateIMManager::Get().SaveAlignmentState();
+
+		// These values are used to size the column container and should not affect its contents.
+		AlignmentState.MinWidth.Reset();
+		AlignmentState.MaxWidth.Reset();
+		AlignmentState.MinHeight.Reset();
+		AlignmentState.MaxHeight.Reset();
+
+		NextTableColumn(ColumnID);
+
+		if (!ColumnLabel.IsEmpty())
+		{
+			FSlateIMManager::Get().RestoreAlignmentState(AlignmentState);
+			SlateIM::Text(ColumnLabel);
+		}
+
+		EndTableColumn();
+	}
+
 	void FixedTableColumnWidth(float Width)
 	{
 		SlateIM::AutoSize();
@@ -341,20 +484,79 @@ namespace SlateIM
 		SlateIM::MinWidth(Width);
 	}
 
-	bool NextTableCell()
+	void BeginTableBody()
 	{
-		SCOPED_NAMED_EVENT_TEXT("SlateIM::NextTableCell", FColorList::Goldenrod);
+		// If there's an open column, close it
+		if (TSharedPtr<SImStackBox> Cell = FSlateIMManager::Get().GetCurrentContainer<SImStackBox>())
+		{
+			EndTableColumn();
+		}
 		
+		// If the header is open, end it
+		if (TSharedPtr<FSlateIMTableHeader> Header = FSlateIMManager::Get().GetCurrentContainer<FSlateIMTableHeader>())
+		{
+			EndTableHeader();
+		}
+
+		if (TSharedPtr<SImTableView> Table = FSlateIMManager::Get().GetCurrentContainer<SImTableView>())
+		{
+			FSlateIMManager::Get().PushContainer(FContainerNode(Table->GetBody()));
+			Table->BeginTableContent();
+		}
+		else
+		{
+			ensureAlwaysMsgf(false, TEXT("Current container should be a table - Is there a missing SlateIM::BeginTable() statement?"));
+		}
+	}
+
+	void EndTableBody()
+	{
 		// If there's an open cell, close it
 		if (TSharedPtr<SImStackBox> Cell = FSlateIMManager::Get().GetCurrentContainer<SImStackBox>())
 		{
 			SlateIM::EndVerticalStack();
 		}
 
+		// If there's an open row, end it
+		if (TSharedPtr<FSlateIMTableRow> Row = FSlateIMManager::Get().GetCurrentContainer<FSlateIMTableRow>())
+		{
+			EndTableRow();
+		}
+
+		FSlateIMManager::Get().PopContainer<FSlateIMTableBody>();
+	}
+
+	bool NextTableCell()
+	{
+		SCOPED_NAMED_EVENT_TEXT("SlateIM::NextTableCell", FColorList::Goldenrod);
+
+		// If there's an open column or cell, close it
+		if (TSharedPtr<SImStackBox> Cell = FSlateIMManager::Get().GetCurrentContainer<SImStackBox>())
+		{
+			if (Cell->GetOrientation() == Orient_Horizontal)
+			{
+				SlateIM::EndHorizontalStack();
+			}
+			else if (Cell->GetOrientation() == Orient_Vertical)
+			{ 
+				SlateIM::EndVerticalStack();
+			}
+		}
+
+		// If we haven't ended the header, do that now.
+		if (TSharedPtr<FSlateIMTableHeader> Header = FSlateIMManager::Get().GetCurrentContainer<FSlateIMTableHeader>())
+		{
+			EndTableHeader();
+		}
+		
 		// We haven't drawn any table cells yet, begin the first row
 		if (TSharedPtr<SImTableView> Table = FSlateIMManager::Get().GetCurrentContainer<SImTableView>())
 		{
-			Table->BeginTableContent();
+			BeginTableBody();
+			BeginTableRow();
+		}
+		else if (FSlateIMManager::Get().GetCurrentContainer<FSlateIMTableBody>())
+		{
 			BeginTableRow();
 		}
 
diff --git a/Engine/Plugins/Experimental/SlateIM/Source/SlateIM/Private/Implementation/SlateIM_Slots.cpp b/Engine/Plugins/Experimental/SlateIM/Source/SlateIM/Private/Implementation/SlateIM_Slots.cpp
--- a/Engine/Plugins/Experimental/SlateIM/Source/SlateIM/Private/Implementation/SlateIM_Slots.cpp
+++ b/Engine/Plugins/Experimental/SlateIM/Source/SlateIM/Private/Implementation/SlateIM_Slots.cpp
@@ -1,4 +1,4 @@
-﻿// Copyright Epic Games, Inc. All Rights Reserved.
+// Copyright Epic Games, Inc. All Rights Reserved.
 
 #include "SlateIM.h"
 
@@ -50,4 +50,11 @@ namespace SlateIM
 	{
 		FSlateIMManager::Get().NextMaxHeight = InMaxHeight;
 	}
+
+	void Maximize()
+	{
+		HAlign(HAlign_Fill);
+		VAlign(VAlign_Fill);
+		Fill();
+	}
 }
diff --git a/Engine/Plugins/Experimental/SlateIM/Source/SlateIM/Private/Misc/SlateIMExamples.cpp b/Engine/Plugins/Experimental/SlateIM/Source/SlateIM/Private/Misc/SlateIMExamples.cpp
--- a/Engine/Plugins/Experimental/SlateIM/Source/SlateIM/Private/Misc/SlateIMExamples.cpp
+++ b/Engine/Plugins/Experimental/SlateIM/Source/SlateIM/Private/Misc/SlateIMExamples.cpp
@@ -576,10 +576,22 @@ void FSlateIMTestWidget::DrawTables()
 			SlateIM::HAlign(HAlign_Fill);
 			SlateIM::MaxHeight(200.f);
 			SlateIM::BeginTable();
-			SlateIM::AddTableColumn(TEXT("Item"), TEXT("Item"));
-			SlateIM::HAlign(HAlign_Center);
+			SlateIM::BeginTableHeader();
+			SlateIM::NextTableColumn(TEXT("Item"));
+			// Makes clickable area fill the column header
+			SlateIM::Maximize();
+			SlateIM::Padding(FMargin(0));
+			SlateIM::BeginContextMenuAnchor();
+			SlateIM::Text(TEXT("Item"));
+			SlateIM::AddMenuSection(TEXT("Menu"));
+			SlateIM::AddMenuButton(TEXT("Click me!"));
+			SlateIM::EndContextMenuAnchor();
 			SlateIM::InitialTableColumnWidth(80.f);
+			SlateIM::Maximize();
+			SlateIM::HAlign(HAlign_Center);
 			SlateIM::AddTableColumn(TEXT("Button"), TEXT("Button"));
+			SlateIM::EndTableHeader();
+			SlateIM::BeginTableBody();
 			for (int32 i = 0; i < NumTableItems; ++i)
 			{
 				if (SlateIM::NextTableCell())
@@ -600,6 +612,7 @@ void FSlateIMTestWidget::DrawTables()
 					}
 				}
 			}
+			SlateIM::EndTableBody();
 			SlateIM::EndTable();
 		}
 		SlateIM::EndHorizontalStack();
diff --git a/Engine/Plugins/Experimental/SlateIM/Source/SlateIM/Private/Misc/SlateIMManager.h b/Engine/Plugins/Experimental/SlateIM/Source/SlateIM/Private/Misc/SlateIMManager.h
--- a/Engine/Plugins/Experimental/SlateIM/Source/SlateIM/Private/Misc/SlateIMManager.h
+++ b/Engine/Plugins/Experimental/SlateIM/Source/SlateIM/Private/Misc/SlateIMManager.h
@@ -106,6 +106,18 @@ namespace SlateIM
 			return RootName == Other.RootName && ContainerIndex == Other.ContainerIndex && WidgetIndex == Other.WidgetIndex;
 		}
 	};
+
+	struct FSlateIMAlignmentState
+	{
+		TOptional<FMargin> Padding;
+		TOptional<EHorizontalAlignment> HAlign;
+		TOptional<EVerticalAlignment> VAlign;
+		TOptional<bool> AutoSize;
+		TOptional<float> MinWidth;
+		TOptional<float> MinHeight;
+		TOptional<float> MaxWidth;
+		TOptional<float> MaxHeight;
+	};
 	
 	class FSlateIMManager
 	{
@@ -178,6 +190,32 @@ namespace SlateIM
 		}
 		void ResetAlignmentData();
 
+		FSlateIMAlignmentState SaveAlignmentState() const
+		{
+			return {
+				NextPadding,
+				NextHAlign,
+				NextVAlign,
+				NextAutoSize,
+				NextMinWidth,
+				NextMinHeight,
+				NextMaxWidth,
+				NextMaxHeight
+			};
+		}
+
+		void RestoreAlignmentState(const FSlateIMAlignmentState& AlignmentState)
+		{
+			NextPadding = AlignmentState.Padding;
+			NextHAlign = AlignmentState.HAlign;
+			NextVAlign = AlignmentState.VAlign;
+			NextAutoSize = AlignmentState.AutoSize;
+			NextMinWidth = AlignmentState.MinWidth;
+			NextMinHeight = AlignmentState.MinHeight;
+			NextMaxWidth = AlignmentState.MaxWidth;
+			NextMaxHeight = AlignmentState.MaxHeight;
+		}
+
 		void OnSlateIMModalOpened();
 		void OnSlateIMModalClosed();
 		bool CanUpdateSlateIM() const;
diff --git a/Engine/Plugins/Experimental/SlateIM/Source/SlateIM/Public/SlateIM.h b/Engine/Plugins/Experimental/SlateIM/Source/SlateIM/Public/SlateIM.h
--- a/Engine/Plugins/Experimental/SlateIM/Source/SlateIM/Public/SlateIM.h
+++ b/Engine/Plugins/Experimental/SlateIM/Source/SlateIM/Public/SlateIM.h
@@ -190,13 +190,30 @@ namespace SlateIM
 	SLATEIM_API void EndTable();
 
 	/**
-	 * Add a column with a label. Call this before adding any content to the table.
-	 * 
+	 * Starts the header. Call this before adding any content to the table.
+	 */
+	SLATEIM_API void BeginTableHeader();
+
+	/**
+	 * Ends the table header.
+	 */
+	SLATEIM_API void EndTableHeader();
+
+	/**
+	 * Moves to the next column. Allows fine control of the column body.
+	 *
 	 * @param ColumnID The name to use as ID for the column
-	 * @param ColumnLabel (optional) The text label to display in the header for this column
 	 *
 	 * @note Table columns support SetToolTip, HAlign, and VAlign. See FixedTableColumnWidth() and InitialTableColumnWidth() for size options.
 	 */
+	SLATEIM_API void NextTableColumn(const FName& ColumnID);
+
+	/**
+	 * Adds a column and immediately ends it with the given label as its body.
+	 *
+	 * @param ColumnID The name to use as ID for the column
+	 * @param ColumnLabel Added a default label to the column
+	 */
 	SLATEIM_API void AddTableColumn(const FName& ColumnID, const FStringView& ColumnLabel = FStringView());
 
 	/**
@@ -213,6 +230,16 @@ namespace SlateIM
 	 */
 	SLATEIM_API void InitialTableColumnWidth(float Width);
 
+	/**
+	 * Starts a table body.
+	 */
+	SLATEIM_API void BeginTableBody();
+
+	/**
+	 * Ends a table body.
+	 */
+	SLATEIM_API void EndTableBody();
+
 	/**
 	 * Begin the content for the next table cell. The column and row are tracked internally based on the number of columns.
 	 * 
@@ -325,6 +352,10 @@ namespace SlateIM
 	 * @param InMaxHeight The maximum height for the next slot
 	 */
 	SLATEIM_API void MaxHeight(float InMaxHeight);
+	/**
+	 * Set the alignments to Fill and removes autosize.
+	 */
+	SLATEIM_API void Maximize();
 #pragma endregion Slots
 
 #pragma region Widgets
