FSlateIMNomadTabBase: Use the tab title and not the console command as the name of the widget

#rnx
#Synced-CL 49104745

diff --git a/Engine/Plugins/Experimental/SlateIM/Source/SlateIM/Private/SlateIMWidgetBase.cpp b/Engine/Plugins/Experimental/SlateIM/Source/SlateIM/Private/SlateIMWidgetBase.cpp
--- a/Engine/Plugins/Experimental/SlateIM/Source/SlateIM/Private/SlateIMWidgetBase.cpp
+++ b/Engine/Plugins/Experimental/SlateIM/Source/SlateIM/Private/SlateIMWidgetBase.cpp
@@ -45,6 +45,11 @@ FSlateIMWidgetWithCommandBase::FSlateIMWidgetWithCommandBase(const TCHAR* Comman
 	, WidgetCommand(Command, CommandHelp, FConsoleCommandDelegate::CreateRaw(this, &FSlateIMWidgetWithCommandBase::ToggleWidget))
 {}
 
+FSlateIMWidgetWithCommandBase::FSlateIMWidgetWithCommandBase(const FStringView& WidgetName, const TCHAR* Command, const TCHAR* CommandHelp)
+	: FSlateIMWidgetBase(WidgetName)
+	, WidgetCommand(Command, CommandHelp, FConsoleCommandDelegate::CreateRaw(this, &FSlateIMWidgetWithCommandBase::ToggleWidget))
+{}
+
 void FSlateIMWindowBase::DrawWidget(float DeltaTime)
 {
 	if (!SlateIM::CanUpdateSlateIM())
@@ -66,7 +71,7 @@ void FSlateIMWindowBase::DrawWidget(float DeltaTime)
 }
 
 FSlateIMNomadTabBase::FSlateIMNomadTabBase(FStringView TabTitle, const TCHAR* Command, const TCHAR* CommandHelp, FSlateIcon TabIcon)
-	: FSlateIMWidgetWithCommandBase(Command, CommandHelp)
+	: FSlateIMWidgetWithCommandBase(TabTitle, Command, CommandHelp)
 {
 	FGlobalTabmanager::Get()->RegisterNomadTabSpawner(GetWidgetName(), FOnSpawnTab::CreateLambda([this, TabTitleText = FText::FromStringView(TabTitle)](const FSpawnTabArgs&)
 	{
diff --git a/Engine/Plugins/Experimental/SlateIM/Source/SlateIM/Public/SlateIMWidgetBase.h b/Engine/Plugins/Experimental/SlateIM/Source/SlateIM/Public/SlateIMWidgetBase.h
--- a/Engine/Plugins/Experimental/SlateIM/Source/SlateIM/Public/SlateIMWidgetBase.h
+++ b/Engine/Plugins/Experimental/SlateIM/Source/SlateIM/Public/SlateIMWidgetBase.h
@@ -49,6 +49,7 @@ class FSlateIMWidgetWithCommandBase : public FSlateIMWidgetBase
 {
 public:
 	UE_API FSlateIMWidgetWithCommandBase(const TCHAR* Command, const TCHAR* CommandHelp);
+	UE_API FSlateIMWidgetWithCommandBase(const FStringView& WidgetName, const TCHAR* Command, const TCHAR* CommandHelp);
 
 private:
 	FAutoConsoleCommand WidgetCommand;
