[SlateIM] be able to use NomadTabs for InGame widgets

#Synced-CL 49037166

diff --git a/Engine/Plugins/Experimental/SlateIM/Source/SlateIMInGame/Private/SlateIMInGameNomadTab.cpp b/Engine/Plugins/Experimental/SlateIM/Source/SlateIMInGame/Private/SlateIMInGameNomadTab.cpp
new file mode 100644
--- /dev/null
+++ b/Engine/Plugins/Experimental/SlateIM/Source/SlateIMInGame/Private/SlateIMInGameNomadTab.cpp
@@ -0,0 +1,103 @@
+// Copyright Epic Games, Inc. All Rights Reserved.
+
+#include "SlateIMInGameNomadTab.h"
+
+#include "SlateIM.h"
+#include "Framework/Docking/TabManager.h"
+#include "Widgets/Docking/SDockTab.h"
+
+#include UE_INLINE_GENERATED_CPP_BY_NAME(SlateIMInGameNomadTab)
+
+ASlateIMInGameNomadTab::ASlateIMInGameNomadTab(const FObjectInitializer& ObjectInitializer)
+{
+
+}
+
+ASlateIMInGameNomadTab::ASlateIMInGameNomadTab(const FName& InTabName, const FStringView& InTabTitle, const FSlateIcon& InTabIcon)
+	: TabName(InTabName)
+	, TabTitle(FText::FromStringView(InTabTitle))
+	, TabIcon(InTabIcon)
+{
+
+}
+
+void ASlateIMInGameNomadTab::OnWidgetStarted()
+{
+	FGlobalTabmanager::Get()->RegisterNomadTabSpawner(TabName, FOnSpawnTab::CreateLambda([this](const FSpawnTabArgs&)
+	{
+		return SNew(SDockTab)
+			.Clipping(EWidgetClipping::ClipToBounds)
+			.TabRole(ETabRole::NomadTab)
+			.Label(TabTitle)
+			.OnTabActivated_Lambda([this](TSharedRef<SDockTab>, ETabActivationCause) 
+				{ OnTabActivated(); })
+			.OnTabClosed_Lambda([this](TSharedRef<SDockTab>) 
+				{ OnTabDeactivated(); });
+	}))
+	.SetIcon(TabIcon);
+}
+
+void ASlateIMInGameNomadTab::OnWidgetStopped()
+{
+	if (TSharedPtr<SDockTab> RootTabPin = RootTab.Pin())
+	{
+		RootTabPin->SetOnTabActivated(SDockTab::FOnTabActivatedCallback());
+		RootTabPin->SetOnTabClosed(SDockTab::FOnTabClosedCallback());
+		RootTabPin->RequestCloseTab();
+		RootTab.Reset();
+		TabContent.Reset();
+	}
+
+	FGlobalTabmanager::Get()->UnregisterNomadTabSpawner(TabName);
+}
+
+void ASlateIMInGameNomadTab::DrawWidget(const float DeltaTime)
+{
+	if (!SlateIM::CanUpdateSlateIM())
+	{
+		return;
+	}
+
+	TSharedPtr<SDockTab> RootTabPin = RootTab.Pin();
+	if (!RootTabPin.IsValid())
+	{
+		// This is the first time we call DrawWidget() since being enabled - invoke the tab
+		RootTab = RootTabPin = FGlobalTabmanager::Get()->TryInvokeTab(TabName);
+		check(RootTabPin.IsValid())
+	}
+
+	TSharedPtr<SWidget> NewExposedWidget;
+	if (SlateIM::BeginExposedRoot(TabName, NewExposedWidget))
+	{
+		DrawContent(DeltaTime);
+	}
+	SlateIM::EndRoot();
+
+	if (NewExposedWidget.IsValid())
+	{
+		if (NewExposedWidget != TabContent)
+		{
+			RootTabPin->SetContent(NewExposedWidget.ToSharedRef());
+			TabContent = NewExposedWidget;
+		}
+	}
+	else
+	{
+		RootTabPin->RequestCloseTab();
+		OnTabDeactivated();
+	}
+}
+
+void ASlateIMInGameNomadTab::OnTabActivated()
+{
+}
+
+void ASlateIMInGameNomadTab::OnTabDeactivated()
+{
+	TabContent.Reset();
+	if (!bDestroyRequested)
+	{
+		bDestroyRequested = true;
+		Server_Destroy();
+	}
+}
diff --git a/Engine/Plugins/Experimental/SlateIM/Source/SlateIMInGame/Private/SlateIMInGameWidgetBase.cpp b/Engine/Plugins/Experimental/SlateIM/Source/SlateIMInGame/Private/SlateIMInGameWidgetBase.cpp
--- a/Engine/Plugins/Experimental/SlateIM/Source/SlateIMInGame/Private/SlateIMInGameWidgetBase.cpp
+++ b/Engine/Plugins/Experimental/SlateIM/Source/SlateIMInGame/Private/SlateIMInGameWidgetBase.cpp
@@ -122,6 +122,8 @@ void ASlateIMInGameWidgetBase::StartWidget()
 	if (!WidgetTickHandle.IsValid())
 	{
 		WidgetTickHandle = FSlateApplication::Get().OnPreTick().AddUObject(this, &ASlateIMInGameWidgetBase::TickWidget);
+
+		OnWidgetStarted();
 	}
 }
 
@@ -131,6 +133,8 @@ void ASlateIMInGameWidgetBase::StopWidget()
 	{
 		FSlateApplication::Get().OnPreTick().Remove(WidgetTickHandle);
 		WidgetTickHandle.Reset();
+
+		OnWidgetStopped();
 	}
 }
 
diff --git a/Engine/Plugins/Experimental/SlateIM/Source/SlateIMInGame/Public/SlateIMInGameNomadTab.h b/Engine/Plugins/Experimental/SlateIM/Source/SlateIMInGame/Public/SlateIMInGameNomadTab.h
new file mode 100644
--- /dev/null
+++ b/Engine/Plugins/Experimental/SlateIM/Source/SlateIMInGame/Public/SlateIMInGameNomadTab.h
@@ -0,0 +1,38 @@
+// Copyright Epic Games, Inc. All Rights Reserved.
+
+#pragma once
+
+#include "SlateIMInGameWidgetBase.h"
+
+#include "SlateIMInGameNomadTab.generated.h"
+
+#define UE_API SLATEIMINGAME_API
+
+UCLASS(MinimalAPI)
+class ASlateIMInGameNomadTab : public ASlateIMInGameWidgetBase
+{
+	GENERATED_BODY()
+public:
+	ASlateIMInGameNomadTab(const FObjectInitializer& ObjectInitializer);
+	UE_API ASlateIMInGameNomadTab(const FName& InTabName, const FStringView& InTabTitle, const FSlateIcon& InTabIcon);
+
+protected:
+	virtual void DrawContent(const float DeltaTime) {};
+	UE_API virtual void OnWidgetStarted() override;
+	UE_API virtual void OnWidgetStopped() override;
+
+private:
+	UE_API virtual void DrawWidget(const float DeltaTime) override;
+
+	void OnTabActivated();
+	void OnTabDeactivated();
+
+	FName TabName;
+	FText TabTitle;
+	FSlateIcon TabIcon;
+	TWeakPtr<SDockTab> RootTab;
+	TWeakPtr<SWidget> TabContent;
+	bool bDestroyRequested = false;
+};
+
+#undef UE_API
diff --git a/Engine/Plugins/Experimental/SlateIM/Source/SlateIMInGame/Public/SlateIMInGameWidgetBase.h b/Engine/Plugins/Experimental/SlateIM/Source/SlateIMInGame/Public/SlateIMInGameWidgetBase.h
--- a/Engine/Plugins/Experimental/SlateIM/Source/SlateIMInGame/Public/SlateIMInGameWidgetBase.h
+++ b/Engine/Plugins/Experimental/SlateIM/Source/SlateIMInGame/Public/SlateIMInGameWidgetBase.h
@@ -39,10 +39,13 @@ protected:
 
 	virtual void DrawWidget(const float DeltaTime) {}
 
-	UE_API virtual void StartWidget();
-	UE_API virtual void StopWidget();
+	UE_API void StartWidget();
+	UE_API void StopWidget();
 	void TickWidget(const float DeltaTime);
 
+	virtual void OnWidgetStarted() {}
+	virtual void OnWidgetStopped() {}
+
 	static ASlateIMInGameWidgetBase* GetOrOpenInGameWidget(APlayerController* Owner, const TSubclassOf<ASlateIMInGameWidgetBase>& InGameWidgetClass);
 	static void DestroyInGameWidget(const APlayerController* Owner, const TSubclassOf<ASlateIMInGameWidgetBase>& InGameWidgetClass);
 
diff --git a/Engine/Plugins/Experimental/SlateIM/Source/SlateIMInGame/SlateIMInGame.Build.cs b/Engine/Plugins/Experimental/SlateIM/Source/SlateIMInGame/SlateIMInGame.Build.cs
--- a/Engine/Plugins/Experimental/SlateIM/Source/SlateIMInGame/SlateIMInGame.Build.cs
+++ b/Engine/Plugins/Experimental/SlateIM/Source/SlateIMInGame/SlateIMInGame.Build.cs
@@ -15,6 +15,7 @@ public class SlateIMInGame : ModuleRules
 				"CoreUObject",
 				"Engine",
 				"Slate",
+				"SlateCore",
 				"SlateIM",
 			}
 		);
