SlateIM - Disable SlateIM updates during package saving

SlateIM widgets can make calls to `StaticFindObject()` while drawing, which is not valid during package save, so we disable SlateIM drawing during that time.

#Synced-CL 49590866

diff --git a/Engine/Plugins/Experimental/SlateIM/Source/SlateIM/Private/Misc/SlateIMManager.cpp b/Engine/Plugins/Experimental/SlateIM/Source/SlateIM/Private/Misc/SlateIMManager.cpp
--- a/Engine/Plugins/Experimental/SlateIM/Source/SlateIM/Private/Misc/SlateIMManager.cpp
+++ b/Engine/Plugins/Experimental/SlateIM/Source/SlateIM/Private/Misc/SlateIMManager.cpp
@@ -344,8 +344,7 @@ namespace SlateIM
 
 	bool FSlateIMManager::CanUpdateSlateIM() const
 	{
-		// Disable SlateIM updates while we have an open modal
-		return !bIsSlateIMModalOpen;
+		return !bIsSlateIMModalOpen && !UE::IsSavingPackage(nullptr);
 	}
 
 	void FSlateIMManager::ValidateRootName(FName RootName) const
