SlateIM - Implement EndSubMenu()

#Synced-CL 49590866

diff --git a/Engine/Plugins/Experimental/SlateIM/Source/SlateIM/Private/Implementation/SlateIM_Menu.cpp b/Engine/Plugins/Experimental/SlateIM/Source/SlateIM/Private/Implementation/SlateIM_Menu.cpp
--- a/Engine/Plugins/Experimental/SlateIM/Source/SlateIM/Private/Implementation/SlateIM_Menu.cpp
+++ b/Engine/Plugins/Experimental/SlateIM/Source/SlateIM/Private/Implementation/SlateIM_Menu.cpp
@@ -85,13 +85,16 @@ namespace SlateIM
 	void BeginSubMenu(const FStringView& SubMenuText)
 	{
 		TSharedPtr<SImContextMenuAnchor> Anchor = FSlateIMManager::Get().GetCurrentMenuRoot();
-		checkf(Anchor, TEXT("Cannot add menu items without a current active menu anchor"));
+		checkf(Anchor, TEXT("Cannot begin a submenu without a current active menu anchor"));
 
 		Anchor->BeginSubMenu(SubMenuText);
 	}
 
 	void EndSubMenu()
 	{
+		TSharedPtr<SImContextMenuAnchor> Anchor = FSlateIMManager::Get().GetCurrentMenuRoot();
+		checkf(Anchor, TEXT("Cannot end a submenu without a current active menu anchor"));
 
+		Anchor->EndSubMenu();
 	}
 }
diff --git a/Engine/Plugins/Experimental/SlateIM/Source/SlateIM/Private/Misc/SlateIMExamples.cpp b/Engine/Plugins/Experimental/SlateIM/Source/SlateIM/Private/Misc/SlateIMExamples.cpp
--- a/Engine/Plugins/Experimental/SlateIM/Source/SlateIM/Private/Misc/SlateIMExamples.cpp
+++ b/Engine/Plugins/Experimental/SlateIM/Source/SlateIM/Private/Misc/SlateIMExamples.cpp
@@ -400,6 +400,14 @@ void FSlateIMTestWidget::DrawBasics()
 			SlateIM::AddMenuButton(TEXT("SubMenu Item 3"), TEXT("Menu Item Tool Tip 3"));
 			SlateIM::AddMenuButton(TEXT("SubMenu Item 4"), TEXT("Menu Item Tool Tip 4"));
 			SlateIM::EndSubMenu();
+			SlateIM::AddMenuButton(TEXT("Menu Item 5"), TEXT("A Menu Item between 2 SubMenus"));
+			SlateIM::AddMenuButton(TEXT("Menu Item 6"), TEXT("A Menu Item between 2 SubMenus"));
+			SlateIM::BeginSubMenu(TEXT("Sub Menu 2"));
+			SlateIM::AddMenuButton(TEXT("SubMenu2 Item 1"), TEXT("Menu Item Tool Tip 1"));
+			SlateIM::AddMenuButton(TEXT("SubMenu2 Item 2"), TEXT("Menu Item Tool Tip 2"));
+			SlateIM::AddMenuButton(TEXT("SubMenu2 Item 3"), TEXT("Menu Item Tool Tip 3"));
+			SlateIM::AddMenuButton(TEXT("SubMenu2 Item 4"), TEXT("Menu Item Tool Tip 4"));
+			SlateIM::EndSubMenu();
 			SlateIM::EndContextMenuAnchor();
 		}
 
