SlateIM - Fix SlateIM widgets in a viewport root not occupying their specified space

#Synced-CL 49621591

diff --git a/Engine/Plugins/Experimental/SlateIM/Source/SlateIM/Private/Roots/SlateIMViewportRoot.cpp b/Engine/Plugins/Experimental/SlateIM/Source/SlateIM/Private/Roots/SlateIMViewportRoot.cpp
--- a/Engine/Plugins/Experimental/SlateIM/Source/SlateIM/Private/Roots/SlateIMViewportRoot.cpp
+++ b/Engine/Plugins/Experimental/SlateIM/Source/SlateIM/Private/Roots/SlateIMViewportRoot.cpp
@@ -90,6 +90,8 @@ void FSlateIMViewportRoot::UpdateViewport(const SlateIM::FViewportRootLayout& La
 				SAssignNew(ScaleBox, SScaleBox)
 				.Stretch(EStretch::UserSpecified)
 				.UserSpecifiedScale(Layout.Scale)
+				.HAlign(HAlign_Fill)
+				.VAlign(VAlign_Fill)
 				[
 					SAssignNew(Wrapper, SImWrapper)	
 				]
