SlateIM: The nomad tabs now enable their widgets before "activation" so they draw before they are focused.

#Synced-CL 49593614

diff --git a/Engine/Plugins/Experimental/SlateIM/Source/SlateIM/Private/SlateIMWidgetBase.cpp b/Engine/Plugins/Experimental/SlateIM/Source/SlateIM/Private/SlateIMWidgetBase.cpp
--- a/Engine/Plugins/Experimental/SlateIM/Source/SlateIM/Private/SlateIMWidgetBase.cpp
+++ b/Engine/Plugins/Experimental/SlateIM/Source/SlateIM/Private/SlateIMWidgetBase.cpp
@@ -75,13 +75,13 @@ FSlateIMNomadTabBase::FSlateIMNomadTabBase(FStringView TabTitle, const TCHAR* Co
 {
 	FGlobalTabmanager::Get()->RegisterNomadTabSpawner(GetWidgetName(), FOnSpawnTab::CreateLambda([this, TabTitleText = FText::FromStringView(TabTitle)](const FSpawnTabArgs&)
 	{
+		EnableWidget();
+
 		return SNew(SDockTab)
 			.Clipping(EWidgetClipping::ClipToBounds)
 			.TabRole(ETabRole::NomadTab)
 			.Label(TabTitleText)
-			.OnTabActivated_Lambda([this](TSharedRef<SDockTab>, ETabActivationCause) 
-				{ EnableWidget(); })
-			.OnTabClosed_Lambda([this](TSharedRef<SDockTab>) 
+			.OnTabClosed_Lambda([this](TSharedRef<SDockTab>)
 				{ FSlateIMWidgetWithCommandBase::DisableWidget(); TabContent.Reset(); });
 	}))
 		.SetIcon(TabIcon);
