SlateIM: There's now a "menu bar" that be added. It can be added anywhere and multiple can exist in a window.

#Synced-CL 49593614

diff --git a/Engine/Plugins/Experimental/SlateIM/Source/SlateIM/Private/Containers/SImContextMenuAnchor.cpp b/Engine/Plugins/Experimental/SlateIM/Source/SlateIM/Private/Containers/SImContextMenuAnchor.cpp
--- a/Engine/Plugins/Experimental/SlateIM/Source/SlateIM/Private/Containers/SImContextMenuAnchor.cpp
+++ b/Engine/Plugins/Experimental/SlateIM/Source/SlateIM/Private/Containers/SImContextMenuAnchor.cpp
@@ -150,12 +150,32 @@ void SImContextMenuAnchor::EndSubMenu()
 	checkf(CurrentSubMenuLevel >= 0, TEXT("Too many calls to EndSubMenu. This means you called EndSubMenu more than BeginSubMenu"));
 }
 
+void SImContextMenuAnchor::OpenMenu(const FVector2D& InScreenPosition, const FWidgetPath* InWidgetPath)
+{
+	if (IsMenuOpen())
+	{
+		return;
+	}
+
+	OpenedMenu = FSlateApplication::Get().PushMenu(
+		AsShared(), 
+		InWidgetPath ? *InWidgetPath : FWidgetPath(), 
+		MenuWidget.ToSharedRef(), 
+		InScreenPosition, 
+		FPopupTransitionEffect(FPopupTransitionEffect::ContextMenu)
+	);
+}
+
+TSharedRef<SWidget> SImContextMenuAnchor::GetMenuContent() const
+{
+	return MenuWidget.IsValid() ? MenuWidget.ToSharedRef() : SNullWidget::NullWidget;
+}
+
 FReply SImContextMenuAnchor::OnMouseButtonUp(const FGeometry& MyGeometry, const FPointerEvent& MouseEvent)
 {
 	if (MenuWidget && MouseEvent.GetEffectingButton() == EKeys::RightMouseButton)
 	{
-		FWidgetPath WidgetPath = MouseEvent.GetEventPath() != nullptr ? *MouseEvent.GetEventPath() : FWidgetPath();
-		OpenedMenu = FSlateApplication::Get().PushMenu(AsShared(), WidgetPath, MenuWidget.ToSharedRef(), MouseEvent.GetScreenSpacePosition(), FPopupTransitionEffect(FPopupTransitionEffect::ContextMenu));
+		OpenMenu(MouseEvent.GetScreenSpacePosition(), MouseEvent.GetEventPath());
 
 		return FReply::Handled();
 	}
diff --git a/Engine/Plugins/Experimental/SlateIM/Source/SlateIM/Private/Containers/SImContextMenuAnchor.h b/Engine/Plugins/Experimental/SlateIM/Source/SlateIM/Private/Containers/SImContextMenuAnchor.h
--- a/Engine/Plugins/Experimental/SlateIM/Source/SlateIM/Private/Containers/SImContextMenuAnchor.h
+++ b/Engine/Plugins/Experimental/SlateIM/Source/SlateIM/Private/Containers/SImContextMenuAnchor.h
@@ -5,6 +5,7 @@
 #include "ISlateIMContainer.h"
 #include "Widgets/SCompoundWidget.h"
 
+class FWidgetPath;
 class IMenu;
 
 enum class EMenuType : uint8
@@ -50,7 +51,9 @@ public:
 	void AddMenuSection(const FStringView& SectionText);
 	void BeginSubMenu(const FStringView& SectionText);
 	void EndSubMenu();
+	void OpenMenu(const FVector2D& InScreenPosition, const FWidgetPath* InWidgetPath = nullptr);
 	bool IsMenuOpen() const { return OpenedMenu.IsValid();  }
+	TSharedRef<SWidget> GetMenuContent() const;
 
 	virtual FReply OnMouseButtonUp(const FGeometry& MyGeometry, const FPointerEvent& MouseEvent) override;
 	
diff --git a/Engine/Plugins/Experimental/SlateIM/Source/SlateIM/Private/Implementation/SlateIM_Menu.cpp b/Engine/Plugins/Experimental/SlateIM/Source/SlateIM/Private/Implementation/SlateIM_Menu.cpp
--- a/Engine/Plugins/Experimental/SlateIM/Source/SlateIM/Private/Implementation/SlateIM_Menu.cpp
+++ b/Engine/Plugins/Experimental/SlateIM/Source/SlateIM/Private/Implementation/SlateIM_Menu.cpp
@@ -1,13 +1,20 @@
-﻿// Copyright Epic Games, Inc. All Rights Reserved.
+// Copyright Epic Games, Inc. All Rights Reserved.
 
 #include "SlateIM.h"
 
 #include "Containers/SImContextMenuAnchor.h"
+#include "Containers/SImStackBox.h"
 #include "Misc/SlateIMManager.h"
 #include "Misc/SlateIMWidgetScope.h"
+#include "SlateIMStyle.h"
 
 namespace SlateIM
 {
+	namespace Menu
+	{
+		const FLazyName MenuBoxTag = "MenuBox";
+	}
+
 	bool BeginContextMenuAnchor()
 	{
 		TSharedPtr<SImContextMenuAnchor> ContainerWidget;
@@ -49,6 +56,83 @@ namespace SlateIM
 
 		return Anchor->AddMenuSeparator();
 	}
+	
+	void BeginMenuBar()
+	{
+		SlateIM::Padding(0);
+		BeginHorizontalStack();
+		
+		// Another class isn't needed here. Tag the stackbox for future checking.
+		TSharedPtr<SImStackBox> MenuBox = FSlateIMManager::Get().GetCurrentContainer<SImStackBox>();
+		MenuBox->SetTag(Menu::MenuBoxTag);
+	}
+
+	void EndMenuBar()
+	{
+		// Pop past open menu if it isn't closed.
+		if (TSharedPtr<SImContextMenuAnchor> MenuAnchor = FSlateIMManager::Get().GetCurrentContainer<SImContextMenuAnchor>())
+		{
+			EndContextMenuAnchor();
+		}
+
+		TSharedPtr<SImStackBox> MenuBox = FSlateIMManager::Get().GetCurrentContainer<SImStackBox>();
+
+		if (!ensureAlwaysMsgf(MenuBox && MenuBox->GetTag() == Menu::MenuBoxTag, TEXT("Attempting to close a menu bar when not in a menu.")))
+		{
+			return;
+		}
+
+		EndHorizontalStack();
+	}
+
+	void AddMenuBarEntry(const FStringView& InMenuName)
+	{
+		// Pop past open menu if it isn't closed.
+		if (TSharedPtr<SImContextMenuAnchor> MenuAnchor = FSlateIMManager::Get().GetCurrentContainer<SImContextMenuAnchor>())
+		{
+			EndContextMenuAnchor();
+		}
+
+		TSharedPtr<SImStackBox> MenuBox = FSlateIMManager::Get().GetCurrentContainer<SImStackBox>();
+
+		if (!ensureAlwaysMsgf(MenuBox && MenuBox->GetTag() == Menu::MenuBoxTag, TEXT("Attempting to add a menu entry when not in a menu.")))
+		{
+			return;
+		}
+
+		SlateIM::Padding(0);
+		BeginContextMenuAnchor();
+
+		SlateIM::Padding(0);
+
+		if (Button(InMenuName, &FSlateIMStyle::Get().GetWidgetStyle<FButtonStyle>("SlateIM.MenuButton")))
+		{
+			if (TSharedPtr<SImContextMenuAnchor> Menu = FSlateIMManager::Get().GetCurrentMenuRoot())
+			{
+				FVector2D Position = Menu->GetTickSpaceGeometry().GetAbsolutePosition();
+				Position.Y += Menu->GetTickSpaceGeometry().GetAbsoluteSize().Y;
+
+				Menu->OpenMenu(Position);
+			}
+		}
+	}
+
+	void EndMenuBarEntry()
+	{
+		TSharedPtr<SImContextMenuAnchor> MenuAnchor = FSlateIMManager::Get().GetCurrentContainer<SImContextMenuAnchor>();
+
+		if (!ensureAlwaysMsgf(MenuAnchor, TEXT("Attempting to close a menu when not in a menu.")))
+		{
+			return;
+		}
+
+		EndContextMenuAnchor();
+	}
+
+	void NextMenuBarEntry(const FStringView& InMenuName)
+	{
+		AddMenuBarEntry(InMenuName);
+	}
 
 	void AddMenuSection(const FStringView& SectionText)
 	{
diff --git a/Engine/Plugins/Experimental/SlateIM/Source/SlateIM/Private/Misc/SlateIMExamples.cpp b/Engine/Plugins/Experimental/SlateIM/Source/SlateIM/Private/Misc/SlateIMExamples.cpp
--- a/Engine/Plugins/Experimental/SlateIM/Source/SlateIM/Private/Misc/SlateIMExamples.cpp
+++ b/Engine/Plugins/Experimental/SlateIM/Source/SlateIM/Private/Misc/SlateIMExamples.cpp
@@ -98,6 +98,15 @@ void FSlateIMTestWidget::Draw()
 
 	SCOPED_NAMED_EVENT_TEXT("FSlateIMTestWidget::Draw", FColorList::Goldenrod);
 	constexpr bool bAbsorbMouse = false;
+
+	SlateIM::HAlign(HAlign_Fill);
+	SlateIM::BeginVerticalStack();
+
+	SlateIM::BeginMenuBar();
+	SlateIM::AddMenuBarEntry(TEXT("Test"));
+	SlateIM::AddMenuButton(TEXT("Option 1"));
+	SlateIM::EndMenuBar();
+
 	SlateIM::BeginBorder(FAppStyle::GetBrush("ToolPanel.GroupBorder"), Orient_Vertical, bAbsorbMouse);
 	// Basic perf measurement, outside of the scrollbox so that it's "pinned"
 	if (TimeSinceLastUpdate > 0.5f)
@@ -105,7 +114,7 @@ void FSlateIMTestWidget::Draw()
 		TimeText = FString::Printf(TEXT("%.3f ms"), LastTime * 1000);
 		TimeSinceLastUpdate = 0;
 	}
-	SlateIM::HAlign(HAlign_Fill);
+
 	SlateIM::BeginHorizontalStack();
 	{
 #if WITH_ENGINE
@@ -159,6 +168,8 @@ void FSlateIMTestWidget::Draw()
 	SlateIM::EndTabStack();
 	SlateIM::EndTabGroup();
 	SlateIM::EndBorder();
+
+	SlateIM::EndVerticalStack();
 }
 
 void FSlateIMTestWidget::DrawBasics()
diff --git a/Engine/Plugins/Experimental/SlateIM/Source/SlateIM/Private/SlateIMStyle.cpp b/Engine/Plugins/Experimental/SlateIM/Source/SlateIM/Private/SlateIMStyle.cpp
new file mode 100644
--- /dev/null
+++ b/Engine/Plugins/Experimental/SlateIM/Source/SlateIM/Private/SlateIMStyle.cpp
@@ -0,0 +1,28 @@
+// Copyright Epic Games, Inc. All Rights Reserved.
+
+#include "SlateIMStyle.h"
+
+#include "Styling/AppStyle.h"
+#include "Styling/SlateStyleRegistry.h"
+#include "Styling/SlateTypes.h"
+
+FSlateIMStyle::FSlateIMStyle()
+	: FSlateStyleSet(TEXT("SlateIM"))
+{
+	const FName ButtonStyleName = "WindowMenuBar.Button";
+
+	FButtonStyle MenuButtonStyle = FAppStyle::GetWidgetStyle<FButtonStyle>(ButtonStyleName);
+	MenuButtonStyle.NormalPadding.Left = 6.f;
+	MenuButtonStyle.NormalPadding.Right = 6.f;
+	MenuButtonStyle.PressedPadding.Left = 6.f;
+	MenuButtonStyle.PressedPadding.Right = 6.f;
+
+	Set("SlateIM.MenuButton", MenuButtonStyle);
+
+	FSlateStyleRegistry::RegisterSlateStyle(*this);
+}
+
+FSlateIMStyle::~FSlateIMStyle()
+{
+	FSlateStyleRegistry::UnRegisterSlateStyle(*this);
+}
diff --git a/Engine/Plugins/Experimental/SlateIM/Source/SlateIM/Private/SlateIMStyle.h b/Engine/Plugins/Experimental/SlateIM/Source/SlateIM/Private/SlateIMStyle.h
new file mode 100644
--- /dev/null
+++ b/Engine/Plugins/Experimental/SlateIM/Source/SlateIM/Private/SlateIMStyle.h
@@ -0,0 +1,18 @@
+// Copyright Epic Games, Inc. All Rights Reserved.
+
+#pragma once
+
+#include "Styling/SlateStyle.h"
+
+class FSlateIMStyle final : public FSlateStyleSet
+{
+public:
+	static FSlateIMStyle& Get()
+	{
+		static FSlateIMStyle Instance;
+		return Instance;
+	}
+
+	FSlateIMStyle();
+	virtual ~FSlateIMStyle() override;
+};
diff --git a/Engine/Plugins/Experimental/SlateIM/Source/SlateIM/Public/SlateIM.h b/Engine/Plugins/Experimental/SlateIM/Source/SlateIM/Public/SlateIM.h
--- a/Engine/Plugins/Experimental/SlateIM/Source/SlateIM/Public/SlateIM.h
+++ b/Engine/Plugins/Experimental/SlateIM/Source/SlateIM/Public/SlateIM.h
@@ -2,7 +2,6 @@
 
 #pragma once
 
-
 #include "Layout/Margin.h"
 #include "Modules/ModuleInterface.h"
 #include "Styling/SlateColor.h"
@@ -585,6 +584,28 @@ namespace SlateIM
 	SLATEIM_API bool BeginContextMenuAnchor();
 	SLATEIM_API void EndContextMenuAnchor();
 
+	/** 
+	* Begin a menu bar, similar to those shown at the top of windows. Can be placed anywhere.
+	* 
+	* @see AddMenuSeparator AddMenuButton AddMenuToggleButton AddMenuCheckButton BeginSubMenu
+	*/
+	SLATEIM_API void BeginMenuBar();
+
+	/**
+	 * End the menu bar.
+	 */
+	SLATEIM_API void EndMenuBar();
+
+	/** 
+	 * Adds a menu to the menu entry with the given name.
+	 */
+	SLATEIM_API void AddMenuBarEntry(const FStringView& InMenuName);
+
+	/**
+	 * Ends a menu bar entry 
+	 */
+	SLATEIM_API void EndMenuBarEntry();
+
 	/**
 	 * Add a new section to the current context menu
 	 * 
