SlateIM: Dynamic tables can now save expansion state correctly.

#Synced-CL 49593614

diff --git a/Engine/Plugins/Experimental/SlateIM/Source/SlateIM/Private/Containers/SImTableView.cpp b/Engine/Plugins/Experimental/SlateIM/Source/SlateIM/Private/Containers/SImTableView.cpp
--- a/Engine/Plugins/Experimental/SlateIM/Source/SlateIM/Private/Containers/SImTableView.cpp
+++ b/Engine/Plugins/Experimental/SlateIM/Source/SlateIM/Private/Containers/SImTableView.cpp
@@ -214,7 +214,12 @@ void FSlateIMTableRow::SetOwningTable(const TSharedPtr<SImTableView>& InOwningTa
 
 bool FSlateIMTableRow::IsExpanded()
 {
-	return OwningTable && OwningTable->IsItemExpanded(AsShared());
+	if (OwningTable)
+	{
+		return OwningTable->IsItemExpanded(AsShared());
+	}
+
+	return false;
 }
 
 bool FSlateIMTableRow::HasChildRows() const
@@ -395,6 +400,35 @@ const TArray<TSharedRef<FSlateIMTableRow>>& FSlateIMTableBody::GetTableRows() co
 	return Children;
 }
 
+bool FSlateIMTableBody::GetSavedExpansionState(FName InRowId)
+{
+	return ExpandedItems.Contains(InRowId);
+}
+
+void FSlateIMTableBody::HandleOnExpansionChanged(TSharedRef<FSlateIMTableRow> InItem, bool bInExpanded)
+{
+	if (!OwningTable.IsValid())
+	{
+		return;
+	}
+
+	const FName RowId = InItem->GetRowId();
+
+	if (RowId.IsNone())
+	{
+		return;
+	}
+
+	if (bInExpanded)
+	{
+		ExpandedItems.Add(RowId);
+	}
+	else
+	{
+		ExpandedItems.Remove(RowId);
+	}
+}
+
 void SImTableView::PrivateRegisterAttributes(FSlateAttributeInitializer&)
 {
 }
@@ -412,7 +446,9 @@ void SImTableView::Construct(const FArguments& InArgs)
 		.TreeItemsSource(&Body->GetTableRows())  // Cannot set IItemsSource here, but still needs data.
 		.HeaderRow(SNew(SImTableHeader, Header.ToSharedRef()))
 		.OnGenerateRow(this, &SImTableView::GenerateRow)
-		.OnGetChildren(this, &SImTableView::GatherChildren));
+		.OnGetChildren(this, &SImTableView::GatherChildren)
+		.OnExpansionChanged(this, &SImTableView::HandleOnExpansionChanged)
+	);
 
 	SetRootItemsSource(MakeUnique<FRowSource>(*this));
 	RequestListRefresh();
@@ -667,6 +703,11 @@ void SImTableView::AddColumn_Internal(const FName& ColumnId, const FStringView&
 	RequestRefresh();
 }
 
+void SImTableView::HandleOnExpansionChanged(TSharedRef<FSlateIMTableRow> InItem, bool bInExpanded)
+{
+	Body->HandleOnExpansionChanged(InItem, bInExpanded);
+}
+
 void SImTableHeader::Construct(const SHeaderRow::FArguments& InArgs, const TSharedRef<FSlateIMTableHeader> InHeader)
 {
 	Header = InHeader;
diff --git a/Engine/Plugins/Experimental/SlateIM/Source/SlateIM/Private/Containers/SImTableView.h b/Engine/Plugins/Experimental/SlateIM/Source/SlateIM/Private/Containers/SImTableView.h
--- a/Engine/Plugins/Experimental/SlateIM/Source/SlateIM/Private/Containers/SImTableView.h
+++ b/Engine/Plugins/Experimental/SlateIM/Source/SlateIM/Private/Containers/SImTableView.h
@@ -2,6 +2,7 @@
 
 #pragma once
 
+#include "Containers/Set.h"
 #include "ISlateIMContainer.h"
 #include "Misc/ISlateIMChild.h"
 #include "Widgets/Views/IItemsSource.h"
@@ -70,6 +71,9 @@ public:
 	virtual void UpdateChild(FSlateIMChild Child, int32 Index, const FSlateIMSlotData& AlignmentData) override;
 	virtual TSharedPtr<SWidget> GetAsWidget() override;
 
+	FName GetRowId() const { return RowId; }
+	void SetRowId(FName InRowId) { RowId = InRowId; }
+
 	void GetChildRows(TArray<TSharedRef<FSlateIMTableRow>>& OutRows) const;
 
 	int32 CountCellWidgetsUpToIndex(int32 Index) const;
@@ -90,6 +94,7 @@ public:
 
 private:
 	// A mix of cell widgets and child rows
+	FName RowId;
 	TArray<FSlateIMChild> Children;
 	int32 ColumnCount = 0;
 	TSharedPtr<SImTableView> OwningTable;
@@ -137,9 +142,14 @@ public:
 
 	const TArray<TSharedRef<FSlateIMTableRow>>& GetTableRows() const;
 
+	bool GetSavedExpansionState(FName InRowId);
+
+	void HandleOnExpansionChanged(TSharedRef<FSlateIMTableRow> InItem, bool bInExpanded);
+
 private:
 	TArray<TSharedRef<FSlateIMTableRow>> Children;
 	TSharedPtr<SImTableView> OwningTable;
+	TSet<FName> ExpandedItems;
 };
 
 // No direct widget representation
@@ -203,6 +213,8 @@ public:
 private:
 	void AddColumn_Internal(const FName& ColumnId, const FStringView& ColumnToolTip, const FSlateIMSlotData& SlotData);
 
+	void HandleOnExpansionChanged(TSharedRef<FSlateIMTableRow> InItem, bool bInExpanded);
+
 	// Because the row source is now in a different object which can change, use this intermediary object 
 	// to avoid issues updating the source.
 	struct FRowSource : public UE::Slate::ItemsSource::IItemsSource<TSharedRef<FSlateIMTableRow>>
diff --git a/Engine/Plugins/Experimental/SlateIM/Source/SlateIM/Private/Implementation/SlateIM_Containers.cpp b/Engine/Plugins/Experimental/SlateIM/Source/SlateIM/Private/Implementation/SlateIM_Containers.cpp
--- a/Engine/Plugins/Experimental/SlateIM/Source/SlateIM/Private/Implementation/SlateIM_Containers.cpp
+++ b/Engine/Plugins/Experimental/SlateIM/Source/SlateIM/Private/Implementation/SlateIM_Containers.cpp
@@ -589,7 +589,7 @@ namespace SlateIM
 		return Row && Row->AreTableRowContentsRequired();
 	}
 
-	bool BeginTableRowChildren()
+	bool BeginTableRowChildren(FName ParentRowId)
 	{
 		// TODO - Let caller set default expansion state
 		// TODO - Save and restore expansion state
@@ -606,8 +606,21 @@ namespace SlateIM
 			return false;
 		}
 
+		ParentRow->SetRowId(ParentRowId);
+
 		SlateIM::BeginTableRow();
 
+		if (!ParentRowId.IsNone())
+		{
+			if (TSharedPtr<SImTableView> OwningTable = ParentRow->GetOwningTable())
+			{
+				if (TSharedPtr<FSlateIMTableBody> TableBody = OwningTable->GetBody())
+				{
+					OwningTable->SetItemExpansion(ParentRow.ToSharedRef(), TableBody->GetSavedExpansionState(ParentRowId));
+				}
+			}
+		}
+
 		return ParentRow->IsExpanded();
 	}
 
diff --git a/Engine/Plugins/Experimental/SlateIM/Source/SlateIM/Private/Misc/SlateIMExamples.cpp b/Engine/Plugins/Experimental/SlateIM/Source/SlateIM/Private/Misc/SlateIMExamples.cpp
--- a/Engine/Plugins/Experimental/SlateIM/Source/SlateIM/Private/Misc/SlateIMExamples.cpp
+++ b/Engine/Plugins/Experimental/SlateIM/Source/SlateIM/Private/Misc/SlateIMExamples.cpp
@@ -638,6 +638,8 @@ void FSlateIMTestWidget::DrawTables()
 		
 		// Tree Example
 		{
+			SlateIM::CheckBox(TEXT("Show Canada?"), bShouldAddCanada);
+
 			SCOPED_NAMED_EVENT_TEXT("Tree Example", FColorList::Goldenrod);
 			const FTableRowStyle* TableRowStyle = &FAppStyle::Get().GetWidgetStyle<FTableRowStyle>("TableView.AlternatingRow");
 			SlateIM::MinWidth(500.f);
@@ -672,69 +674,72 @@ void FSlateIMTestWidget::DrawTables()
 				}
 				if (SlateIM::BeginTableRowChildren())
 				{
-					if (SlateIM::NextTableCell())
-					{
-						SlateIM::VAlign(VAlign_Center);
-						SlateIM::Text(TEXT("Canada"));
-					}
-					if (SlateIM::NextTableCell())
-					{
-						SlateIM::VAlign(VAlign_Center);
-						SlateIM::Text(TEXT("Country"));
-					}
-					if (SlateIM::BeginTableRowChildren())
+					if (bShouldAddCanada)
 					{
 						if (SlateIM::NextTableCell())
 						{
 							SlateIM::VAlign(VAlign_Center);
-							SlateIM::Text(TEXT("British Columbia"));
+							SlateIM::Text(TEXT("Canada"));
 						}
 						if (SlateIM::NextTableCell())
 						{
 							SlateIM::VAlign(VAlign_Center);
-							SlateIM::Text(TEXT("Province"));
+							SlateIM::Text(TEXT("Country"));
 						}
-						if (SlateIM::BeginTableRowChildren())
+						if (SlateIM::BeginTableRowChildren("CanadaChildren"))
 						{
 							if (SlateIM::NextTableCell())
 							{
 								SlateIM::VAlign(VAlign_Center);
-								SlateIM::Text(TEXT("Vancouver"));
+								SlateIM::Text(TEXT("British Columbia"));
 							}
 							if (SlateIM::NextTableCell())
 							{
 								SlateIM::VAlign(VAlign_Center);
-								SlateIM::Text(TEXT("City"));
+								SlateIM::Text(TEXT("Province"));
 							}
-						}
-						SlateIM::EndTableRowChildren();
-				
-						if (SlateIM::NextTableCell())
-						{
-							SlateIM::VAlign(VAlign_Center);
-							SlateIM::Text(TEXT("Quebec"));
-						}
-						if (SlateIM::NextTableCell())
-						{
-							SlateIM::VAlign(VAlign_Center);
-							SlateIM::Text(TEXT("Province"));
-						}
-						if (SlateIM::BeginTableRowChildren())
-						{
+							if (SlateIM::BeginTableRowChildren())
+							{
+								if (SlateIM::NextTableCell())
+								{
+									SlateIM::VAlign(VAlign_Center);
+									SlateIM::Text(TEXT("Vancouver"));
+								}
+								if (SlateIM::NextTableCell())
+								{
+									SlateIM::VAlign(VAlign_Center);
+									SlateIM::Text(TEXT("City"));
+								}
+							}
+							SlateIM::EndTableRowChildren();
+
 							if (SlateIM::NextTableCell())
 							{
 								SlateIM::VAlign(VAlign_Center);
-								SlateIM::Text(TEXT("Montreal"));
+								SlateIM::Text(TEXT("Quebec"));
 							}
 							if (SlateIM::NextTableCell())
 							{
 								SlateIM::VAlign(VAlign_Center);
-								SlateIM::Text(TEXT("City"));
+								SlateIM::Text(TEXT("Province"));
+							}
+							if (SlateIM::BeginTableRowChildren())
+							{
+								if (SlateIM::NextTableCell())
+								{
+									SlateIM::VAlign(VAlign_Center);
+									SlateIM::Text(TEXT("Montreal"));
+								}
+								if (SlateIM::NextTableCell())
+								{
+									SlateIM::VAlign(VAlign_Center);
+									SlateIM::Text(TEXT("City"));
+								}
 							}
+							SlateIM::EndTableRowChildren();
 						}
 						SlateIM::EndTableRowChildren();
 					}
-					SlateIM::EndTableRowChildren();
 
 					if (SlateIM::NextTableCell())
 					{
@@ -746,7 +751,7 @@ void FSlateIMTestWidget::DrawTables()
 						SlateIM::VAlign(VAlign_Center);
 						SlateIM::Text(TEXT("Country"));
 					}
-					if (SlateIM::BeginTableRowChildren())
+					if (SlateIM::BeginTableRowChildren("UnitedStatesChildren"))
 					{
 						if (SlateIM::NextTableCell())
 						{
diff --git a/Engine/Plugins/Experimental/SlateIM/Source/SlateIM/Private/Misc/SlateIMExamples.h b/Engine/Plugins/Experimental/SlateIM/Source/SlateIM/Private/Misc/SlateIMExamples.h
--- a/Engine/Plugins/Experimental/SlateIM/Source/SlateIM/Private/Misc/SlateIMExamples.h
+++ b/Engine/Plugins/Experimental/SlateIM/Source/SlateIM/Private/Misc/SlateIMExamples.h
@@ -41,6 +41,7 @@ private:
 	int32 LiveNumItems = NumItems;
 	FString NumItemsText = FString::FromInt(NumItems);
 	bool bShouldLiveUpdateTable = false;
+	bool bShouldAddCanada = true;
 	bool bShouldAddChildRow = false;
 	TArray<FString> ComboBoxItems = { TEXT("Option 1"), TEXT("Option 2"), TEXT("Option 3"), TEXT("Option 4") };
 	bool bRefreshComboItems = false;
diff --git a/Engine/Plugins/Experimental/SlateIM/Source/SlateIM/Public/SlateIM.h b/Engine/Plugins/Experimental/SlateIM/Source/SlateIM/Public/SlateIM.h
--- a/Engine/Plugins/Experimental/SlateIM/Source/SlateIM/Public/SlateIM.h
+++ b/Engine/Plugins/Experimental/SlateIM/Source/SlateIM/Public/SlateIM.h
@@ -250,10 +250,14 @@ namespace SlateIM
 
 	/**
 	 * Begin adding content as a child to the current table row. NextTableCell() must still be called after this function.
+	 * In the case of a dynamic table, providing a parent row id ensures that its expansion state is maintained correctly.
+	 * A parent row id is not required if the table is not dynamic.
+	 * 
+	 * @param Id of the parent row.
 	 * 
 	 * @return Whether the parent row is expanded or not. When false, the child content can be skipped to save cycles.
 	 */
-	SLATEIM_API bool BeginTableRowChildren();
+	SLATEIM_API bool BeginTableRowChildren(FName ParentRowId = NAME_None);
 	/**
 	 * Stop adding child content to the parent row. This must be called no matter the result of BeginTableRowChildren()
 	 */
