SlateIM - Fix assert when opening a modal (like a File Dialog) from a SlateIM widget

#Synced-CL 49727878

diff --git a/Engine/Plugins/Experimental/SlateIM/Source/SlateIM/Private/Misc/SlateIMManager.cpp b/Engine/Plugins/Experimental/SlateIM/Source/SlateIM/Private/Misc/SlateIMManager.cpp
--- a/Engine/Plugins/Experimental/SlateIM/Source/SlateIM/Private/Misc/SlateIMManager.cpp
+++ b/Engine/Plugins/Experimental/SlateIM/Source/SlateIM/Private/Misc/SlateIMManager.cpp
@@ -344,7 +344,7 @@ namespace SlateIM
 
 	bool FSlateIMManager::CanUpdateSlateIM() const
 	{
-		return !bIsSlateIMModalOpen && !UE::IsSavingPackage(nullptr);
+		return !bIsSlateIMModalOpen && !UE::IsSavingPackage(nullptr) && FSlateApplication::IsInitialized() && FSlateApplication::Get().GetNumExternalModalWindowsActive() == 0;
 	}
 
 	void FSlateIMManager::ValidateRootName(FName RootName) const
diff --git a/Engine/Source/Runtime/Slate/Public/Framework/Application/SlateApplication.h b/Engine/Source/Runtime/Slate/Public/Framework/Application/SlateApplication.h
--- a/Engine/Source/Runtime/Slate/Public/Framework/Application/SlateApplication.h
+++ b/Engine/Source/Runtime/Slate/Public/Framework/Application/SlateApplication.h
@@ -563,6 +563,8 @@ public:
 	 */
 	SLATE_API void ExternalModalStop();
 
+	int GetNumExternalModalWindowsActive() const { return NumExternalModalWindowsActive; }
+
 	/** Event before slate application ticks. */
 	DECLARE_EVENT_OneParam(FSlateApplication, FSlateTickEvent, float);
 	FSlateTickEvent& OnPreTick()  { return PreTickEvent; }
