SlateIM: You can now create always on top windows.

#Synced-CL 49748899

diff --git a/Engine/Plugins/Experimental/SlateIM/Source/SlateIM/Private/Implementation/SlateIM_Roots.cpp b/Engine/Plugins/Experimental/SlateIM/Source/SlateIM/Private/Implementation/SlateIM_Roots.cpp
--- a/Engine/Plugins/Experimental/SlateIM/Source/SlateIM/Private/Implementation/SlateIM_Roots.cpp
+++ b/Engine/Plugins/Experimental/SlateIM/Source/SlateIM/Private/Implementation/SlateIM_Roots.cpp
@@ -1,7 +1,8 @@
-﻿// Copyright Epic Games, Inc. All Rights Reserved.
+// Copyright Epic Games, Inc. All Rights Reserved.
 
 #include "SlateIM.h"
 
+#include "Containers/Set.h"
 #include "Framework/Application/SlateApplication.h"
 #include "Misc/SlateIMManager.h"
 #include "Roots/SlateIMExposedRoot.h"
@@ -16,6 +17,18 @@
 namespace SlateIM
 {
 	bool BeginWindowRoot(FName UniqueName, const FStringView& WindowTitle, FVector2f WindowSize, bool bShouldReopen)
+	{
+		return BeginWindowRoot(
+			UniqueName,
+			{
+				.WindowTitle   = WindowTitle,
+				.WindowSize    = WindowSize,
+				.bShouldReopen = bShouldReopen
+			}
+		);
+	}
+	
+	bool BeginWindowRoot(FName UniqueName, const FWindowParams& Params)
 	{
 		if (!FSlateIMManager::Get().CanUpdateSlateIM())
 		{
@@ -28,7 +41,7 @@ namespace SlateIM
 		if (RootNode)
 		{
 			// Window existed and was destroyed and needs to be reopened if requested
-			bNeedsCreate = !RootNode->RootWidget.IsValid() && bShouldReopen;
+			bNeedsCreate = !RootNode->RootWidget.IsValid() && Params.bShouldReopen;
 
 			RootNode->RootState = RootNode->RootWidget->IsVisible();
 		}
@@ -39,8 +52,9 @@ namespace SlateIM
 			// TODO - Save and restore window size and position, mapped to UniqueName
 			TSharedRef<SWindow> NewWindowWidget
 				= SNew(SWindow)
-				.Title(FText::FromStringView(WindowTitle))
-				.ClientSize(WindowSize);
+				.Title(FText::FromStringView(Params.WindowTitle))
+				.ClientSize(Params.WindowSize)
+				.IsTopmostWindow(Params.bAlwaysOnTop);
 
 			FSlateApplication::Get().AddWindow(NewWindowWidget);
 
@@ -61,7 +75,7 @@ namespace SlateIM
 			WindowRoot = StaticCastSharedPtr<FSlateIMWindowRoot>(RootNode->RootWidget);
 		}
 		
-		WindowRoot->UpdateWindow(WindowTitle);
+		WindowRoot->UpdateWindow(Params.WindowTitle);
 
 		FSlateIMManager::Get().BeginRoot(UniqueName);
 
diff --git a/Engine/Plugins/Experimental/SlateIM/Source/SlateIM/Private/Misc/SlateIMExamples.h b/Engine/Plugins/Experimental/SlateIM/Source/SlateIM/Private/Misc/SlateIMExamples.h
--- a/Engine/Plugins/Experimental/SlateIM/Source/SlateIM/Private/Misc/SlateIMExamples.h
+++ b/Engine/Plugins/Experimental/SlateIM/Source/SlateIM/Private/Misc/SlateIMExamples.h
@@ -69,7 +69,7 @@ class FSlateStyleBrowser : public FSlateIMWindowBase
 {
 public:
 	FSlateStyleBrowser()
-		: FSlateIMWindowBase(TEXT("SlateIM Style Browser"), FVector2f(1000, 500), TEXT("SlateIM.ToggleSlateStyleBrowser"), TEXT("Opens a window that previews available slate styles"))
+		: FSlateIMWindowBase(TEXT("SlateIM Style Browser"), FVector2f(1000, 500), /* Always On Top */ false, TEXT("SlateIM.ToggleSlateStyleBrowser"), TEXT("Opens a window that previews available slate styles"))
 	{}
 	
 	virtual void DrawWindow(float DeltaTime) override;
@@ -87,7 +87,7 @@ class FSlateIMTestWindowWidget : public FSlateIMWindowBase
 {
 public:
 	FSlateIMTestWindowWidget(const TCHAR* Command, const TCHAR* CommandHelp)
-		: FSlateIMWindowBase(TEXT("SlateIM Test Suite"), FVector2f(700, 900), Command, CommandHelp)
+		: FSlateIMWindowBase(TEXT("SlateIM Test Suite"), FVector2f(700, 900), /* Always On Top */ false, Command, CommandHelp)
 	{}
 	
 	virtual void DrawWindow(float DeltaTime) override;
diff --git a/Engine/Plugins/Experimental/SlateIM/Source/SlateIM/Private/SlateIMWidgetBase.cpp b/Engine/Plugins/Experimental/SlateIM/Source/SlateIM/Private/SlateIMWidgetBase.cpp
--- a/Engine/Plugins/Experimental/SlateIM/Source/SlateIM/Private/SlateIMWidgetBase.cpp
+++ b/Engine/Plugins/Experimental/SlateIM/Source/SlateIM/Private/SlateIMWidgetBase.cpp
@@ -57,7 +57,15 @@ void FSlateIMWindowBase::DrawWidget(float DeltaTime)
 		return;
 	}
 	
-	const bool bIsDrawingWindow = SlateIM::BeginWindowRoot(GetWidgetName(), WindowTitle, WindowSize);
+	const bool bIsDrawingWindow = SlateIM::BeginWindowRoot(
+		GetWidgetName(),
+		{
+			.WindowTitle  = WindowTitle,
+			.WindowSize   = WindowSize,
+			.bAlwaysOnTop = bAlwaysOnTop
+		}
+	);
+
 	if (bIsDrawingWindow)
 	{
 		DrawWindow(DeltaTime);
diff --git a/Engine/Plugins/Experimental/SlateIM/Source/SlateIM/Public/SlateIM.h b/Engine/Plugins/Experimental/SlateIM/Source/SlateIM/Public/SlateIM.h
--- a/Engine/Plugins/Experimental/SlateIM/Source/SlateIM/Public/SlateIM.h
+++ b/Engine/Plugins/Experimental/SlateIM/Source/SlateIM/Public/SlateIM.h
@@ -1,7 +1,6 @@
 // Copyright Epic Games, Inc. All Rights Reserved.
 
 #pragma once
-
 #include "Layout/Margin.h"
 #include "Modules/ModuleInterface.h"
 #include "Styling/SlateColor.h"
@@ -75,6 +74,8 @@ namespace SlateIM
 #endif // WITH_EDITOR || WITH_ENGINE
 
 	/**
+	 * DEPRECATED! Please use the alternate method below.
+	 * 
 	 * Begins a new floating window root
 	 *
 	 * @param UniqueName a unique name to identify this window
@@ -86,8 +87,36 @@ namespace SlateIM
 	 *
 	 * @note Consider using FSlateIMWindowBase to handle creating and updating your Slate IM Window
 	 */
+	UE_DEPRECATED(5.8, "Please use the alternate method with the struct parameter.")
 	SLATEIM_API bool BeginWindowRoot(FName UniqueName, const FStringView& WindowTitle, FVector2f WindowSize, bool bShouldReopen = false);
 
+	struct FWindowParams
+	{
+		// The name of the window. This will also be displayed in the windows title bar
+		FStringView WindowTitle;
+
+		// The size to open the window at (does not update the size of an existing window)
+		FVector2f WindowSize = FVector2f::ZeroVector;
+
+		// If this window was created with a previous call to BeginWindow and then closed, passing in true will reopen the window and false will leave it closed. If the window has never been seen this parameter does nothing and the window will open
+		bool bShouldReopen = false;
+
+		// Window will be create always on top.
+		bool bAlwaysOnTop = false;
+	};
+
+	/**
+	 * Begins a new floating window root
+	 *
+	 * @param UniqueName a unique name to identify this window
+	 * @param Params Configuration parameters for the window.
+	 *
+	 * @return the current window state. True if open, false if closed or not updating
+	 *
+	 * @note Consider using FSlateIMWindowBase to handle creating and updating your Slate IM Window
+	 */
+	SLATEIM_API bool BeginWindowRoot(FName UniqueName, const FWindowParams& Params = {});
+
 #if WITH_ENGINE
 	// TODO - Game viewport widgets might need something to control HitTestability and cursor visibility
 	/**
diff --git a/Engine/Plugins/Experimental/SlateIM/Source/SlateIM/Public/SlateIMWidgetBase.h b/Engine/Plugins/Experimental/SlateIM/Source/SlateIM/Public/SlateIMWidgetBase.h
--- a/Engine/Plugins/Experimental/SlateIM/Source/SlateIM/Public/SlateIMWidgetBase.h
+++ b/Engine/Plugins/Experimental/SlateIM/Source/SlateIM/Public/SlateIMWidgetBase.h
@@ -61,10 +61,11 @@ private:
 class FSlateIMWindowBase : public FSlateIMWidgetWithCommandBase
 {
 public:
-	FSlateIMWindowBase(FStringView WindowTitle, FVector2f WindowSize, const TCHAR* Command, const TCHAR* CommandHelp)
+	FSlateIMWindowBase(FStringView WindowTitle, FVector2f WindowSize, bool bAlwaysOnTop, const TCHAR* Command, const TCHAR* CommandHelp)
 		: FSlateIMWidgetWithCommandBase(Command, CommandHelp)
 		, WindowTitle(WindowTitle)
 		, WindowSize(WindowSize)
+		, bAlwaysOnTop(bAlwaysOnTop)
 	{}
 
 protected:
@@ -75,6 +76,7 @@ private:
 
 	FString WindowTitle;
 	FVector2f WindowSize;
+	bool bAlwaysOnTop;
 };
 
 /**
diff --git a/Engine/Plugins/Experimental/SlateIM/Source/SlateIMInGame/Private/SlateIMInGameWindow.cpp b/Engine/Plugins/Experimental/SlateIM/Source/SlateIMInGame/Private/SlateIMInGameWindow.cpp
--- a/Engine/Plugins/Experimental/SlateIM/Source/SlateIMInGame/Private/SlateIMInGameWindow.cpp
+++ b/Engine/Plugins/Experimental/SlateIM/Source/SlateIMInGame/Private/SlateIMInGameWindow.cpp
@@ -18,7 +18,14 @@ ASlateIMInGameWindow::ASlateIMInGameWindow(const FName& InWindowName, const FStr
 
 void ASlateIMInGameWindow::DrawWidget(const float DeltaTime)
 {
-	const bool bIsDrawingWindow = SlateIM::BeginWindowRoot(WindowName, WindowTitle, WindowSize);
+	const bool bIsDrawingWindow = SlateIM::BeginWindowRoot(
+		WindowName,
+		{
+			.WindowTitle = WindowTitle,
+			.WindowSize = WindowSize
+		}
+	);
+
 	if (bIsDrawingWindow)
 	{
 		DrawContent(DeltaTime);
diff --git a/Engine/Plugins/Runtime/AudioMotorSim/Source/AudioMotorSimDebug/Private/AudioMotorModelDebugWidgetChildWindow.h b/Engine/Plugins/Runtime/AudioMotorSim/Source/AudioMotorSimDebug/Private/AudioMotorModelDebugWidgetChildWindow.h
--- a/Engine/Plugins/Runtime/AudioMotorSim/Source/AudioMotorSimDebug/Private/AudioMotorModelDebugWidgetChildWindow.h
+++ b/Engine/Plugins/Runtime/AudioMotorSim/Source/AudioMotorSimDebug/Private/AudioMotorModelDebugWidgetChildWindow.h
@@ -27,7 +27,13 @@ private:
     		return;
     	}
     	
-    	const bool bIsDrawingWindow = SlateIM::BeginWindowRoot(GetWidgetName(), WindowTitle, WindowSize);
+    	const bool bIsDrawingWindow = SlateIM::BeginWindowRoot(
+			GetWidgetName(),
+			{
+				.WindowTitle = WindowTitle,
+				.WindowSize = WindowSize
+			});
+
     	if (bIsDrawingWindow)
     	{
     		DrawWindow(DeltaTime);
