SlateIM: You can now add a canvas render target image and draw to it via an event.

#Synced-CL 49784486

diff --git a/Engine/Plugins/Experimental/SlateIM/Source/SlateIM/Private/Implementation/SlateIMCanvasDrawCommandList.cpp b/Engine/Plugins/Experimental/SlateIM/Source/SlateIM/Private/Implementation/SlateIMCanvasDrawCommandList.cpp
new file mode 100644
--- /dev/null
+++ b/Engine/Plugins/Experimental/SlateIM/Source/SlateIM/Private/Implementation/SlateIMCanvasDrawCommandList.cpp
@@ -0,0 +1,61 @@
+// Copyright Epic Games, Inc. All Rights Reserved.
+
+#include "Implementation/SlateIMCanvasDrawCommandList.h"
+
+#include "SlateIM.h"
+
+USlateIMCanvasDrawCommandList::USlateIMCanvasDrawCommandList()
+	: UpdateType(SlateIM::Canvas::ECanvasUpdateType::EveryFrame)
+	, bInvalidated(true)
+{
+}
+
+SlateIM::Canvas::ECanvasUpdateType USlateIMCanvasDrawCommandList::GetUpdateType() const
+{
+	return UpdateType;
+}
+
+void USlateIMCanvasDrawCommandList::SetUpdateType(SlateIM::Canvas::ECanvasUpdateType InUpdateType)
+{
+	UpdateType = InUpdateType;
+}
+
+void USlateIMCanvasDrawCommandList::Invalidate()
+{
+	bInvalidated = true;
+}
+
+bool USlateIMCanvasDrawCommandList::NeedsUpdate() const
+{
+	switch (UpdateType)
+	{
+		case SlateIM::Canvas::ECanvasUpdateType::Invalidation:
+			return bInvalidated;
+
+		case SlateIM::Canvas::ECanvasUpdateType::EveryFrame:
+			return true;
+
+		default:
+			ensureAlwaysMsgf(false, TEXT("Invalid update time [%d]."), static_cast<uint8>(UpdateType));
+			return false;
+	}
+}
+
+void USlateIMCanvasDrawCommandList::EnqueueCommand(FCanvasDrawCommand&& InDrawCommand)
+{
+	DrawCommands.Add(Forward<FCanvasDrawCommand>(InDrawCommand));
+}
+
+void USlateIMCanvasDrawCommandList::ProcessCommands(UCanvas* InCanvas, int32 InWidth, int32 InHeight)
+{
+	for (const FCanvasDrawCommand& DrawCommand : DrawCommands)
+	{
+		DrawCommand.ExecuteIfBound(InCanvas);
+	}
+
+	// Make sure all delegates are deleted
+	DrawCommands.Reset();
+
+	// Reset invalidation as we've now been drawn.
+	bInvalidated = false;
+}
diff --git a/Engine/Plugins/Experimental/SlateIM/Source/SlateIM/Private/Implementation/SlateIMCanvasDrawCommandList.h b/Engine/Plugins/Experimental/SlateIM/Source/SlateIM/Private/Implementation/SlateIMCanvasDrawCommandList.h
new file mode 100644
--- /dev/null
+++ b/Engine/Plugins/Experimental/SlateIM/Source/SlateIM/Private/Implementation/SlateIMCanvasDrawCommandList.h
@@ -0,0 +1,46 @@
+// Copyright Epic Games, Inc. All Rights Reserved.
+
+#pragma once
+
+#include "UObject/Object.h"
+
+#include "Containers/Array.h"
+#include "Delegates/DelegateCombinations.h"
+
+#include "SlateIMCanvasDrawCommandList.generated.h"
+
+class UCanvas;
+
+namespace SlateIM::Canvas
+{
+	enum class ECanvasUpdateType : uint8;
+}
+
+UCLASS()
+class USlateIMCanvasDrawCommandList : public UObject
+{
+	GENERATED_BODY()
+
+public:
+	DECLARE_DELEGATE_OneParam(FCanvasDrawCommand, UCanvas*)
+
+	USlateIMCanvasDrawCommandList();
+
+	SlateIM::Canvas::ECanvasUpdateType GetUpdateType() const;
+
+	void SetUpdateType(SlateIM::Canvas::ECanvasUpdateType InUpdateType);
+
+	void Invalidate();
+
+	bool NeedsUpdate() const;
+
+	void EnqueueCommand(FCanvasDrawCommand&& InDrawCommand);
+
+	UFUNCTION()
+	void ProcessCommands(UCanvas* InCanvas, int32 InWidth, int32 InHeight);
+
+protected:
+	SlateIM::Canvas::ECanvasUpdateType UpdateType;
+	bool bInvalidated = false;
+	TArray<FCanvasDrawCommand> DrawCommands;
+};
diff --git a/Engine/Plugins/Experimental/SlateIM/Source/SlateIM/Private/Implementation/SlateIM_Canvas.cpp b/Engine/Plugins/Experimental/SlateIM/Source/SlateIM/Private/Implementation/SlateIM_Canvas.cpp
new file mode 100644
--- /dev/null
+++ b/Engine/Plugins/Experimental/SlateIM/Source/SlateIM/Private/Implementation/SlateIM_Canvas.cpp
@@ -0,0 +1,493 @@
+// Copyright Epic Games, Inc. All Rights Reserved.
+
+// All implemented functions are WITH_ENGINE
+#if WITH_ENGINE
+
+#include "SlateIM.h"
+#include "SlateIM_Internal.h"
+
+#include "Containers/SImCompoundWidget.h"
+#include "Engine/Canvas.h"
+#include "Engine/CanvasRenderTarget2D.h"
+#include "Engine/Texture2D.h"
+#include "Engine/TextureRenderTarget2D.h"
+#include "Implementation/SlateIMCanvasDrawCommandList.h"
+#include "Math/IntPoint.h"
+#include "Misc/SlateIMManager.h"
+#include "Misc/SlateIMSlateResources.h"
+#include "SlateMaterialBrush.h"
+#include "Styling/SlateColor.h"
+#include "UObject/Object.h"
+#include "UObject/UObjectGlobals.h"
+#include "Widgets/Images/SImage.h"
+
+namespace SlateIM::Canvas
+{
+	static const FLazyName CanvasContainerTag = "Canvas";
+
+	const FSlateBrush* GetCanvasBrush(const TSharedRef<SImage>& CanvasImage, TWeakObjectPtr<UObject> WorldContextWeak, int32 Width, int32 Height)
+	{
+		// Create canvas render target
+		UCanvasRenderTarget2D* CanvasRenderTarget = nullptr;
+		TSharedPtr<FSlateIMPinnedImageResourceDataStore> CanvasResource = CanvasImage->GetMetaData<FSlateIMPinnedImageResourceDataStore>();
+
+		if (!CanvasResource)
+		{
+			CanvasRenderTarget = UCanvasRenderTarget2D::CreateCanvasRenderTarget2D(
+				WorldContextWeak.Get(),
+				UCanvasRenderTarget2D::StaticClass(),
+				Width,
+				Height
+			);
+
+			if (!CanvasRenderTarget)
+			{
+				return nullptr;
+			}
+
+			FSlateBrush Brush;
+			Brush.SetResourceObject(CanvasRenderTarget);
+			Brush.ImageSize = FVector2D(Width, Height);
+
+			CanvasResource = MakeShared<FSlateIMPinnedImageResourceDataStore>(FPinnedImageResource{ Brush, TStrongObjectPtr<UObject>(CanvasRenderTarget) });
+			CanvasImage->AddMetadata(CanvasResource.ToSharedRef());
+		}
+		else
+		{
+			CanvasRenderTarget = Cast<UCanvasRenderTarget2D>(CanvasResource->Data.PinnedResource.Get());
+
+			if (!CanvasRenderTarget)
+			{
+				return nullptr;
+			}
+
+			FIntPoint CurrentSize;
+			CanvasRenderTarget->GetSize(CurrentSize.X, CurrentSize.Y);
+
+			if (CurrentSize.X != Width || CurrentSize.Y != Height)
+			{
+				CanvasRenderTarget->ResizeTarget(Width, Height);
+				CanvasResource->Data.Brush.SetImageSize(FVector2D(Width, Height));
+			}
+		}
+
+		// Create uobject proxy for queueing draw calls
+		USlateIMCanvasDrawCommandList* UpdateContainer = nullptr;
+		TSharedPtr<FSlateIMPinnedObjectResourceDataStore> UpdateResource = CanvasImage->GetMetaData<FSlateIMPinnedObjectResourceDataStore>();
+
+		if (!UpdateResource)
+		{
+			UpdateContainer = NewObject<USlateIMCanvasDrawCommandList>();
+			UpdateResource = MakeShared<FSlateIMPinnedObjectResourceDataStore>(FPinnedObjectResource{ TStrongObjectPtr<UObject>(UpdateContainer) });
+			CanvasImage->AddMetadata(UpdateResource.ToSharedRef());
+		}
+		else
+		{
+			UpdateContainer = CastChecked<USlateIMCanvasDrawCommandList>(UpdateResource->Data.PinnedResource.Get());
+		}
+
+		CanvasRenderTarget->OnCanvasRenderTargetUpdate.Clear();
+		CanvasRenderTarget->OnCanvasRenderTargetUpdate.AddDynamic(UpdateContainer, &USlateIMCanvasDrawCommandList::ProcessCommands);
+
+		// Defer update call to EndCanvas
+
+		return &CanvasResource->Data.Brush;
+	}
+
+	USlateIMCanvasDrawCommandList* GetCanvasCommandList()
+	{
+		TSharedPtr<SImCompoundWidget> CanvasContainer = FSlateIMManager::Get().GetCurrentContainer<SImCompoundWidget>();
+
+		if (CanvasContainer && CanvasContainer->GetTag() == CanvasContainerTag)
+		{
+			if (TSharedPtr<SImage> CanvasImage = StaticCastSharedPtr<SImage>(CanvasContainer->GetChild(0).GetWidget()))
+			{
+				if (TSharedPtr<FSlateIMPinnedObjectResourceDataStore> CanvasResource = CanvasImage->GetMetaData<FSlateIMPinnedObjectResourceDataStore>())
+				{
+					return CastChecked<USlateIMCanvasDrawCommandList>(CanvasResource->Data.PinnedResource.Get());
+				}
+			}
+		}
+
+		ensureAlwaysMsgf(false, TEXT("Current container should be a canvas - Is there a missing SlateIM::Canvas::BeginCanvas() statement?"));
+
+		return nullptr;
+	}
+
+	void BeginCanvas(int32 Width, int32 Height, const FCanvasParams& CanvasParams)
+	{
+		BeginContainer();
+
+		// Another class isn't needed here. Tag the container for future checking.
+		TSharedPtr<SImCompoundWidget> CanvasContainer = FSlateIMManager::Get().GetCurrentContainer<SImCompoundWidget>();
+		CanvasContainer->SetTag(CanvasContainerTag);
+		CanvasContainer->SetBackgroundImage(nullptr);
+		CanvasContainer->SetAbsorbMouse(false);
+
+		auto GetBrush = [WorldContextWeak = TWeakObjectPtr<UObject>(CanvasParams.WorldContext), Width, Height]
+			(const TSharedRef<SImage>& ImageWidget) -> const FSlateBrush*
+			{
+				return GetCanvasBrush(ImageWidget, WorldContextWeak, Width, Height);
+			};
+
+		Image_Internal(CanvasParams.ColorAndOpacity, CanvasParams.DesiredSize, GetBrush);
+
+		if (USlateIMCanvasDrawCommandList* CommandList = GetCanvasCommandList())
+		{
+			CommandList->SetUpdateType(CanvasParams.UpdateType);
+		}
+	}
+
+	void EndCanvas()
+	{
+		if (USlateIMCanvasDrawCommandList* CommandList = GetCanvasCommandList())
+		{
+			if (!CommandList->NeedsUpdate())
+			{
+				EndContainer();
+				return;
+			}
+		}
+
+		TSharedPtr<SImCompoundWidget> CanvasContainer = FSlateIMManager::Get().GetCurrentContainer<SImCompoundWidget>();
+
+		if (CanvasContainer && CanvasContainer->GetTag() == CanvasContainerTag)
+		{
+			bool bUpdated = false;
+
+			if (TSharedPtr<SImage> CanvasImage = StaticCastSharedPtr<SImage>(CanvasContainer->GetChild(0).GetWidget()))
+			{
+				if (TSharedPtr<FSlateIMPinnedImageResourceDataStore> CanvasResource = CanvasImage->GetMetaData<FSlateIMPinnedImageResourceDataStore>())
+				{
+					if (UCanvasRenderTarget2D* CanvasRenderTarget = CastChecked<UCanvasRenderTarget2D>(CanvasResource->Data.PinnedResource.Get()))
+					{
+						CanvasRenderTarget->UpdateResourceWithParams(UTexture::EUpdateResourceFlags::None);
+						bUpdated = true;
+					}
+				}
+			}
+
+			if (!bUpdated)
+			{
+				ensureAlwaysMsgf(false, TEXT("Internal error."));
+			}
+
+			EndContainer();
+		}
+		else
+		{
+			ensureAlwaysMsgf(false, TEXT("Current container should be a canvas - Is there a missing SlateIM::Canvas::BeginCanvas() statement?"));
+		}
+	}
+
+	void Invalidate()
+	{
+		if (USlateIMCanvasDrawCommandList* CommandList = GetCanvasCommandList())
+		{
+			CommandList->Invalidate();
+		}
+	}
+
+	void SetClip(const FVector2f& ClipPosition)
+	{
+		if (USlateIMCanvasDrawCommandList* CommandList = GetCanvasCommandList())
+		{
+			CommandList->EnqueueCommand(USlateIMCanvasDrawCommandList::FCanvasDrawCommand::CreateLambda(
+				[ClipPosition](UCanvas* InCanvas)
+				{
+					InCanvas->SetClip(ClipPosition.X, ClipPosition.Y);
+				}
+			));
+		}
+	}
+
+	void SetDrawColor(const FLinearColor& Color)
+	{
+		if (USlateIMCanvasDrawCommandList* CommandList = GetCanvasCommandList())
+		{
+			CommandList->EnqueueCommand(USlateIMCanvasDrawCommandList::FCanvasDrawCommand::CreateLambda(
+				[Color](UCanvas* InCanvas)
+				{
+					InCanvas->SetLinearDrawColor(Color);
+				}
+			));
+		}
+	}
+
+	void SetDrawColor(const FColor& Color)
+	{
+		if (USlateIMCanvasDrawCommandList* CommandList = GetCanvasCommandList())
+		{
+			CommandList->EnqueueCommand(USlateIMCanvasDrawCommandList::FCanvasDrawCommand::CreateLambda(
+				[Color](UCanvas* InCanvas)
+				{
+					InCanvas->SetDrawColor(Color);
+				}
+			));
+		}
+	}
+
+	void DrawItem(TSharedRef<FCanvasItem> ItemPtr, const FVector2f& CanvasPosition)
+	{
+		if (USlateIMCanvasDrawCommandList* CommandList = GetCanvasCommandList())
+		{
+			CommandList->EnqueueCommand(USlateIMCanvasDrawCommandList::FCanvasDrawCommand::CreateLambda(
+				[ItemPtr, CanvasPosition](UCanvas* InCanvas)
+				{
+					InCanvas->DrawItem(*ItemPtr, CanvasPosition.X, CanvasPosition.Y);
+				}
+			));
+		}
+	}
+
+	void DrawIcon(const FCanvasIcon& Icon, const FVector2f& CanvasPosition, const FVector2f& Scale)
+	{
+		if (USlateIMCanvasDrawCommandList* CommandList = GetCanvasCommandList())
+		{
+			CommandList->EnqueueCommand(USlateIMCanvasDrawCommandList::FCanvasDrawCommand::CreateLambda(
+				[Icon, CanvasPosition, Scale](UCanvas* InCanvas)
+				{
+					InCanvas->DrawScaledIcon(
+						Icon, 
+						CanvasPosition.X, 
+						CanvasPosition.Y, 
+						FVector(Scale.X, Scale.Y, 1)
+					);
+				}
+			));
+		}
+	}
+
+	void DrawLine(const FVector2D& CanvasPositionA, const FVector2D& CanvasPositionB, float Thickness, FLinearColor RenderColor)
+	{
+		if (USlateIMCanvasDrawCommandList* CommandList = GetCanvasCommandList())
+		{
+			CommandList->EnqueueCommand(USlateIMCanvasDrawCommandList::FCanvasDrawCommand::CreateLambda(
+				[CanvasPositionA, CanvasPositionB, Thickness, RenderColor](UCanvas* InCanvas)
+				{
+					InCanvas->K2_DrawLine(
+						CanvasPositionA, 
+						CanvasPositionB, 
+						Thickness,
+						RenderColor
+					);
+				}
+			));
+		}
+	}
+
+	void DrawTexture(UTexture* RenderTexture, const FVector2D& CanvasPosition, const FVector2D& CanvasSize, const FTileRenderParams& RenderParams)
+	{
+		if (USlateIMCanvasDrawCommandList* CommandList = GetCanvasCommandList())
+		{
+			CommandList->EnqueueCommand(USlateIMCanvasDrawCommandList::FCanvasDrawCommand::CreateLambda(
+				[RenderTexture, CanvasPosition, CanvasSize, RenderParams](UCanvas* InCanvas)
+				{
+					InCanvas->K2_DrawTexture(
+						RenderTexture, 
+						CanvasPosition,
+						CanvasSize,
+						RenderParams.UVPosition,
+						RenderParams.UVSize,
+						RenderParams.RenderColor,
+						RenderParams.BlendMode,
+						RenderParams.Rotation,
+						RenderParams.PivotPoint
+					);
+				}
+			));
+		}
+	}
+
+	void DrawMaterial(UMaterialInterface* RenderMaterial, const FVector2D& CanvasPosition, const FVector2D& CanvasSize, const FTileRenderParams& RenderParams)
+	{
+		if (USlateIMCanvasDrawCommandList* CommandList = GetCanvasCommandList())
+		{
+			CommandList->EnqueueCommand(USlateIMCanvasDrawCommandList::FCanvasDrawCommand::CreateLambda(
+				[RenderMaterial, CanvasPosition, CanvasSize, RenderParams](UCanvas* InCanvas)
+				{
+					InCanvas->K2_DrawMaterial(
+						RenderMaterial,
+						CanvasPosition,
+						CanvasSize,
+						RenderParams.UVPosition,
+						RenderParams.UVSize,
+						RenderParams.Rotation,
+						RenderParams.PivotPoint
+					);
+				}
+			));
+		}
+	}
+
+	void DrawText(UFont* RenderFont, const FString& RenderText, const FVector2D& CanvasPosition, const FTextRenderParams& RenderParams)
+	{
+		if (USlateIMCanvasDrawCommandList* CommandList = GetCanvasCommandList())
+		{
+			CommandList->EnqueueCommand(USlateIMCanvasDrawCommandList::FCanvasDrawCommand::CreateLambda(
+				[RenderFont, RenderText, CanvasPosition, RenderParams](UCanvas* InCanvas)
+				{
+					InCanvas->K2_DrawText(
+						RenderFont,
+						RenderText,
+						CanvasPosition,
+						RenderParams.Scale,
+						RenderParams.RenderColor,
+						RenderParams.Kerning,
+						RenderParams.ShadowColor,
+						RenderParams.ShadowOffset,
+						RenderParams.bCentreX,
+						RenderParams.bCentreY,
+						RenderParams.bOutlined,
+						RenderParams.OutlineColor
+					);
+				}
+			));
+		}
+	}
+
+	void DrawText(UFont* RenderFont, const FString& RenderText, const FVector2f& CanvasPosition, const FVector2f& Scale, const FFontRenderInfo& RenderInfo)
+	{
+		if (USlateIMCanvasDrawCommandList* CommandList = GetCanvasCommandList())
+		{
+			CommandList->EnqueueCommand(USlateIMCanvasDrawCommandList::FCanvasDrawCommand::CreateLambda(
+				[RenderFont, RenderText, CanvasPosition, Scale, RenderInfo](UCanvas* InCanvas)
+				{
+					InCanvas->DrawText(
+						RenderFont,
+						RenderText,
+						CanvasPosition.X,
+						CanvasPosition.Y,
+						Scale.X,
+						Scale.Y,
+						RenderInfo
+					);
+				}
+			));
+		}
+	}
+
+	void DrawWrappedText(const UFont* RenderFont, const FString& RenderText, const FVector2f& CanvasPosition, bool bCenterTextX, bool bCenterTextY, const FVector2f& Scale, const FFontRenderInfo& RenderInfo)
+	{
+		if (USlateIMCanvasDrawCommandList* CommandList = GetCanvasCommandList())
+		{
+			CommandList->EnqueueCommand(USlateIMCanvasDrawCommandList::FCanvasDrawCommand::CreateLambda(
+				[RenderFont, RenderText, CanvasPosition, Scale, bCenterTextX, bCenterTextY, RenderInfo](UCanvas* InCanvas)
+				{
+					int32 outXL, outYL;
+
+					InCanvas->WrappedPrint(
+						/* Draw */ true,
+						CanvasPosition.X,
+						CanvasPosition.Y,
+						outXL,
+						outYL,
+						RenderFont,
+						Scale.X,
+						Scale.Y,
+						bCenterTextX,
+						bCenterTextY,
+						RenderText,
+						RenderInfo
+					);
+				}
+			));
+		}
+	}
+
+	void DrawBorder(UTexture* BorderTexture, UTexture* BackgroundTexture, UTexture* LeftBorderTexture, UTexture* RightBorderTexture, UTexture* TopBorderTexture, UTexture* BottomBorderTexture, const FVector2D& CanvasPosition, const FVector2D& CanvasSize, const FBorderRenderParams& RenderParams)
+	{
+		if (USlateIMCanvasDrawCommandList* CommandList = GetCanvasCommandList())
+		{
+			CommandList->EnqueueCommand(USlateIMCanvasDrawCommandList::FCanvasDrawCommand::CreateLambda(
+				[BorderTexture,BackgroundTexture, LeftBorderTexture, RightBorderTexture, TopBorderTexture, BottomBorderTexture, 
+					CanvasPosition, CanvasSize, RenderParams](UCanvas* InCanvas)
+				{
+					InCanvas->K2_DrawBorder(
+						BorderTexture,
+						BackgroundTexture,
+						LeftBorderTexture,
+						RightBorderTexture,
+						TopBorderTexture,
+						BottomBorderTexture,
+						CanvasPosition,
+						CanvasSize,
+						RenderParams.UVPosition,
+						RenderParams.UVSize,
+						RenderParams.RenderColor,
+						RenderParams.BorderScale,
+						RenderParams.BackgroundScale,
+						RenderParams.Rotation,
+						RenderParams.PivotPoint,
+						RenderParams.CornerSize
+					);
+				}
+			));
+		}
+	}
+
+	void DrawBox(const FVector2D& CanvasPosition, const FVector2D& CanvasSize, float Thickness, FLinearColor RenderColor)
+	{
+		if (USlateIMCanvasDrawCommandList* CommandList = GetCanvasCommandList())
+		{
+			CommandList->EnqueueCommand(USlateIMCanvasDrawCommandList::FCanvasDrawCommand::CreateLambda(
+				[CanvasPosition, CanvasSize, Thickness, RenderColor](UCanvas* InCanvas)
+				{
+					InCanvas->K2_DrawBox(
+						CanvasPosition,
+						CanvasSize,
+						Thickness,
+						RenderColor
+					);
+				}
+			));
+		}
+	}
+
+	void DrawTriangles(UTexture* RenderTexture, TSharedRef<TArray<FCanvasUVTri>> Triangles)
+	{
+		if (USlateIMCanvasDrawCommandList* CommandList = GetCanvasCommandList())
+		{
+			CommandList->EnqueueCommand(USlateIMCanvasDrawCommandList::FCanvasDrawCommand::CreateLambda(
+				[RenderTexture, Triangles](UCanvas* InCanvas)
+				{
+					InCanvas->K2_DrawTriangle(RenderTexture, *Triangles);
+				}
+			));
+		}
+	}
+
+	void DrawMaterialTriangles(UMaterialInterface* RenderMaterial, TSharedRef<TArray<FCanvasUVTri>> Triangles)
+	{
+		if (USlateIMCanvasDrawCommandList* CommandList = GetCanvasCommandList())
+		{
+			CommandList->EnqueueCommand(USlateIMCanvasDrawCommandList::FCanvasDrawCommand::CreateLambda(
+				[RenderMaterial, Triangles](UCanvas* InCanvas)
+				{
+					InCanvas->K2_DrawMaterialTriangle(RenderMaterial, *Triangles);
+				}
+			));
+		}
+	}
+
+	void DrawPolygon(UTexture* RenderTexture, const FVector2D& CanvasPosition, const FVector2D& Radius, int32 NumberOfSides, FLinearColor RenderColor)
+	{
+		if (USlateIMCanvasDrawCommandList* CommandList = GetCanvasCommandList())
+		{
+			CommandList->EnqueueCommand(USlateIMCanvasDrawCommandList::FCanvasDrawCommand::CreateLambda(
+				[RenderTexture, CanvasPosition, Radius, NumberOfSides, RenderColor](UCanvas* InCanvas)
+				{
+					InCanvas->K2_DrawPolygon(
+						RenderTexture, 
+						CanvasPosition,
+						Radius, 
+						NumberOfSides, 
+						RenderColor
+					);
+				}
+			));
+		}
+	}
+} // SlateIM::Canvas
+
+#endif // WITH_ENGINE
diff --git a/Engine/Plugins/Experimental/SlateIM/Source/SlateIM/Private/Implementation/SlateIM_Containers.cpp b/Engine/Plugins/Experimental/SlateIM/Source/SlateIM/Private/Implementation/SlateIM_Containers.cpp
--- a/Engine/Plugins/Experimental/SlateIM/Source/SlateIM/Private/Implementation/SlateIM_Containers.cpp
+++ b/Engine/Plugins/Experimental/SlateIM/Source/SlateIM/Private/Implementation/SlateIM_Containers.cpp
@@ -15,6 +15,30 @@
 
 namespace SlateIM
 {
+	void BeginContainer()
+	{
+		TSharedPtr<SImCompoundWidget> ContainerWidget;
+
+		// Note widget scope must fall out of scope before we push container or else we are pushing this new container into itself instead of into its parent
+		{
+			FWidgetScope<SImCompoundWidget> Scope;
+			ContainerWidget = Scope.GetWidget();
+
+			if (!ContainerWidget)
+			{
+				ContainerWidget = SNew(SImCompoundWidget);
+				Scope.UpdateWidget(ContainerWidget);
+			}
+		}
+
+		FSlateIMManager::Get().PushContainer(FContainerNode(ContainerWidget));
+	}
+
+	void EndContainer()
+	{
+		FSlateIMManager::Get().PopContainer<SImCompoundWidget>();
+	}
+
 	void BeginStack(EOrientation Orientation)
 	{
 		TSharedPtr<SImStackBox> ContainerWidget;
diff --git a/Engine/Plugins/Experimental/SlateIM/Source/SlateIM/Private/Implementation/SlateIM_Widgets.cpp b/Engine/Plugins/Experimental/SlateIM/Source/SlateIM/Private/Implementation/SlateIM_Widgets.cpp
--- a/Engine/Plugins/Experimental/SlateIM/Source/SlateIM/Private/Implementation/SlateIM_Widgets.cpp
+++ b/Engine/Plugins/Experimental/SlateIM/Source/SlateIM/Private/Implementation/SlateIM_Widgets.cpp
@@ -1,6 +1,7 @@
-﻿// Copyright Epic Games, Inc. All Rights Reserved.
+// Copyright Epic Games, Inc. All Rights Reserved.
 
 #include "SlateIM.h"
+#include "SlateIM_Internal.h"
 
 #include "Misc/SlateIMLogging.h"
 #include "Misc/SlateIMManager.h"
@@ -20,27 +21,12 @@
 #if WITH_ENGINE
 #include "Engine/Texture2D.h"
 #include "Engine/TextureRenderTarget2D.h"
+#include "Misc/SlateIMSlateResources.h"
 #include "SlateMaterialBrush.h"
 #endif
 
 namespace SlateIM
 {
-	template<typename DataType>
-	struct FSlateIMDataStore : ISlateMetaData
-	{
-		SLATE_METADATA_TYPE(FSlateIMDataStore, ISlateMetaData);
-
-		FSlateIMDataStore(DataType&& InData)
-			: Data(MoveTemp(InData))
-		{}
-
-		FSlateIMDataStore(const DataType& InData)
-			: Data(InData)
-		{}
-
-		DataType Data;
-	};
-	
 	void Text(const FStringView& InText, const FTextBlockStyle* TextStyle /*= nullptr*/)
 	{
 		Text(InText, FSlateColor::UseForeground(), TextStyle);
@@ -179,12 +165,6 @@ namespace SlateIM
 	}
 
 #if WITH_ENGINE
-	struct FPinnedImageResource
-	{
-		FSlateBrush Brush;
-		TStrongObjectPtr<UObject> PinnedResource;
-	};
-	using FUObjectImageResource = FSlateIMDataStore<FPinnedImageResource>;
 	void Image(UTexture2D* ImageTexture, const FSlateColor& ColorAndOpacity, FVector2D DesiredSize)
 	{
 		auto MakeTextureBrush = [ImageTexture]() -> FSlateBrush
diff --git a/Engine/Plugins/Experimental/SlateIM/Source/SlateIM/Private/Misc/SlateIMExamples.cpp b/Engine/Plugins/Experimental/SlateIM/Source/SlateIM/Private/Misc/SlateIMExamples.cpp
--- a/Engine/Plugins/Experimental/SlateIM/Source/SlateIM/Private/Misc/SlateIMExamples.cpp
+++ b/Engine/Plugins/Experimental/SlateIM/Source/SlateIM/Private/Misc/SlateIMExamples.cpp
@@ -19,8 +19,10 @@
 #include "SLevelViewport.h"
 #endif
 #if WITH_ENGINE
+#include "Engine/Canvas.h"
 #include "Engine/Engine.h"
 #include "Engine/Texture2D.h"
+#include "Materials/Material.h"
 #endif
 
 namespace SlateIM::Private
@@ -910,6 +912,10 @@ void FSlateIMTestWidget::DrawTables()
 
 void FSlateIMTestWidget::DrawGraphs()
 {
+	SlateIM::VAlign(VAlign_Fill);
+	SlateIM::Padding(0);
+	SlateIM::BeginVerticalStack();
+
 	SCOPED_NAMED_EVENT_TEXT("Graphs", FColorList::Goldenrod);
 	SlateIM::BeginHorizontalStack();
 	{
@@ -974,6 +980,115 @@ void FSlateIMTestWidget::DrawGraphs()
 		SlateIM::EndVerticalStack();
 	}
 	SlateIM::EndHorizontalStack();
+
+#if WITH_ENGINE
+	SCOPED_NAMED_EVENT_TEXT("CanvasRenderTarget", FColorList::Goldenrod);
+
+	FCanvasIcon Icon;
+	Icon.Texture = GreenIcon.LoadSynchronous();
+
+	SlateIM::Canvas::BeginCanvas(100, 100, { .UpdateType = SlateIM::Canvas::ECanvasUpdateType::Invalidation });
+
+	// Allows for caching of textures and materials after the first draw call.
+	if (CanvasRenderCount < 2)
+	{
+		SlateIM::Canvas::Invalidate();
+		++CanvasRenderCount;
+	}
+
+	SlateIM::Canvas::DrawPolygon(
+		nullptr,
+		{ 50, 100 },
+		{ 35, 35 },
+		36
+	);
+
+	SlateIM::Canvas::DrawIcon(Icon, { 18, 65 });
+	SlateIM::Canvas::DrawLine({ 0, 60 }, { 100, 60 }, 1.f, FLinearColor::Red);
+
+	SlateIM::Canvas::DrawTexture(
+		RedIcon.LoadSynchronous(),
+		{ 75, 10 },
+		{ 30, 40 },
+		{ .UVSize = {0.75, 1} }
+	);
+
+	UMaterial* MaterialLeft = TSoftObjectPtr<UMaterial>(FSoftObjectPath(TEXT("/Script/Engine.Material'/Engine/EngineMaterials/ColorGradingWheel.ColorGradingWheel'"))).LoadSynchronous();
+	
+	SlateIM::Canvas::DrawMaterial(
+		MaterialLeft,
+		{ 0, 10 },
+		{ 30, 40 },
+		{ .UVPosition = {0.25, 0}, .UVSize = {0.75, 1} }
+	);
+
+	UFont* EngineFont = TSoftObjectPtr<UFont>(FSoftObjectPath(TEXT("/Script/Engine.Font'/Engine/EngineFonts/RobotoTiny.RobotoTiny'"))).LoadSynchronous();
+
+	SlateIM::Canvas::DrawText(
+		EngineFont,
+		TEXT("Canvas"),
+		{ 50, 15 },
+		{
+			.Scale = {1.0, 1.0},
+			.ShadowColor = FLinearColor(0.5, 0.5, 0.5, 1.f),
+			.bCentreX = true
+		}
+	);
+
+	SlateIM::Canvas::DrawText(
+		EngineFont,
+		TEXT("Hi!"),
+		{ 40, 30 },
+		{ 1.5, 1.5 }
+	);
+
+	SlateIM::Canvas::DrawBox(
+		{ 0, 0 },
+		{ 100, 100 },
+		2,
+		FLinearColor::Blue
+	);
+
+	UTexture2D* TriangleTexture = TSoftObjectPtr<UTexture2D>(FSoftObjectPath(TEXT("/Script/Engine.Texture2D'/Engine/EngineMaterials/ColorGradingWheelGradient.ColorGradingWheelGradient'"))).LoadSynchronous();
+
+	TSharedRef<TArray<FCanvasUVTri>> TextureTriangles = MakeShared<TArray<FCanvasUVTri>>();
+	{
+		FCanvasUVTri& Triangle = TextureTriangles->AddDefaulted_GetRef();
+		Triangle.V0_Color = FLinearColor::White;
+		Triangle.V0_UV = { 0, 0 };
+		Triangle.V0_Pos = { 5, 70 };
+		Triangle.V1_Color = FLinearColor::White;
+		Triangle.V1_UV = { 0, 1 };
+		Triangle.V1_Pos = { 25, 70 };
+		Triangle.V2_Color = FLinearColor::White;
+		Triangle.V2_UV = { 1, 0 };
+		Triangle.V2_Pos = { 5, 90 };
+	}
+
+	SlateIM::Canvas::DrawTriangles(TriangleTexture, TextureTriangles);
+
+	UMaterial* TriangleMaterial = TSoftObjectPtr<UMaterial>(FSoftObjectPath(TEXT("/Script/Engine.Material'/Engine/EngineMaterials/ColorGradingWheel.ColorGradingWheel'"))).LoadSynchronous();
+
+	TSharedRef<TArray<FCanvasUVTri>> MaterialTriangles = MakeShared<TArray<FCanvasUVTri>>();
+	{
+		FCanvasUVTri& Triangle = MaterialTriangles->AddDefaulted_GetRef();
+		Triangle.V0_Color = FLinearColor::White;
+		Triangle.V0_UV = { 0, 0 };
+		Triangle.V0_Pos = { 75, 70 };
+		Triangle.V1_Color = FLinearColor::White;
+		Triangle.V1_UV = { 1, 0 };
+		Triangle.V1_Pos = { 95, 70 };
+		Triangle.V2_Color = FLinearColor::White;
+		Triangle.V2_UV = { 1, 1 };
+		Triangle.V2_Pos = { 95, 90 };
+	}
+
+	SlateIM::Canvas::DrawMaterialTriangles(TriangleMaterial, MaterialTriangles);
+
+	SlateIM::Canvas::EndCanvas();
+#endif
+
+	SlateIM::EndVerticalStack();
 }
 
 void FSlateIMTestWidget::DrawInputs()
diff --git a/Engine/Plugins/Experimental/SlateIM/Source/SlateIM/Private/Misc/SlateIMExamples.h b/Engine/Plugins/Experimental/SlateIM/Source/SlateIM/Private/Misc/SlateIMExamples.h
--- a/Engine/Plugins/Experimental/SlateIM/Source/SlateIM/Private/Misc/SlateIMExamples.h
+++ b/Engine/Plugins/Experimental/SlateIM/Source/SlateIM/Private/Misc/SlateIMExamples.h
@@ -62,6 +62,7 @@ private:
 #if WITH_ENGINE
 	TSoftObjectPtr<UTexture2D> RedIcon = TSoftObjectPtr<UTexture2D>(FSoftObjectPath(TEXT("/Engine/EngineResources/AICON-Red.AICON-Red")));
 	TSoftObjectPtr<UTexture2D> GreenIcon = TSoftObjectPtr<UTexture2D>(FSoftObjectPath(TEXT("/Engine/EngineResources/AICON-Green.AICON-Green")));
+	int32 CanvasRenderCount = 0;
 #endif
 };
 
diff --git a/Engine/Plugins/Experimental/SlateIM/Source/SlateIM/Private/Misc/SlateIMSlateResources.h b/Engine/Plugins/Experimental/SlateIM/Source/SlateIM/Private/Misc/SlateIMSlateResources.h
new file mode 100644
--- /dev/null
+++ b/Engine/Plugins/Experimental/SlateIM/Source/SlateIM/Private/Misc/SlateIMSlateResources.h
@@ -0,0 +1,91 @@
+// Copyright Epic Games, Inc. All Rights Reserved.
+ 
+#pragma once
+
+// Entire file is WITH_ENGINE
+#if WITH_ENGINE
+
+#include "Styling/SlateBrush.h"
+#include "Types/ISlateMetaData.h"
+#include "UObject/Object.h"
+#include "UObject/StrongObjectPtr.h"
+
+namespace SlateIM
+{
+	template<typename DataType>
+	struct FSlateIMDataStore : ISlateMetaData
+	{
+		SLATE_METADATA_TYPE(FSlateIMDataStore, ISlateMetaData);
+
+		FSlateIMDataStore(DataType&& InData)
+			: Data(MoveTemp(InData))
+		{
+		}
+
+		FSlateIMDataStore(const DataType& InData)
+			: Data(InData)
+		{
+		}
+
+		DataType Data;
+	};
+
+	/**
+	 * Object Resource
+	 */
+	struct FPinnedObjectResource
+	{
+		TStrongObjectPtr<UObject> PinnedResource;
+	};
+
+	using FUObjectObjectResource = FSlateIMDataStore<FPinnedObjectResource>;
+
+	struct FSlateIMPinnedObjectResourceDataStore : FSlateIMDataStore<FPinnedObjectResource>
+	{
+		using Super = FSlateIMDataStore<FPinnedObjectResource>;
+
+		FSlateIMPinnedObjectResourceDataStore(FPinnedObjectResource&& InData)
+			: Super(MoveTemp(InData))
+		{
+		}
+
+		FSlateIMPinnedObjectResourceDataStore(const FPinnedObjectResource& InData)
+			: Super(InData)
+		{
+		}
+
+		static const FName& GetTypeId() { static FName Type(TEXT("FSlateIMPinnedObjectResourceDataStore")); return Type; }
+		virtual bool IsOfTypeImpl(const FName& Type) const override { return GetTypeId() == Type || Super::IsOfTypeImpl(Type); }
+	};
+
+	/**
+	 * Image Resource
+	 */
+	struct FPinnedImageResource
+	{
+		FSlateBrush Brush;
+		TStrongObjectPtr<UObject> PinnedResource;
+	};
+
+	using FUObjectImageResource = FSlateIMDataStore<FPinnedImageResource>;
+
+	struct FSlateIMPinnedImageResourceDataStore : FSlateIMDataStore<FPinnedImageResource>
+	{
+		using Super = FSlateIMDataStore<FPinnedImageResource>;
+
+		FSlateIMPinnedImageResourceDataStore(FPinnedImageResource&& InData)
+			: Super(MoveTemp(InData))
+		{
+		}
+
+		FSlateIMPinnedImageResourceDataStore(const FPinnedImageResource& InData)
+			: Super(InData)
+		{
+		}
+
+		static const FName& GetTypeId() { static FName Type(TEXT("FSlateIMPinnedImageResourceDataStore")); return Type; }
+		virtual bool IsOfTypeImpl(const FName& Type) const override { return GetTypeId() == Type || Super::IsOfTypeImpl(Type); }
+	};
+}
+
+#endif // WITH_ENGINE
diff --git a/Engine/Plugins/Experimental/SlateIM/Source/SlateIM/Private/SlateIM_Internal.h b/Engine/Plugins/Experimental/SlateIM/Source/SlateIM/Private/SlateIM_Internal.h
new file mode 100644
--- /dev/null
+++ b/Engine/Plugins/Experimental/SlateIM/Source/SlateIM/Private/SlateIM_Internal.h
@@ -0,0 +1,26 @@
+// Copyright Epic Games, Inc. All Rights Reserved.
+
+#pragma once
+
+#include "Math/MathFwd.h"
+#include "Templates/FunctionFwd.h"
+#include "Templates//SharedPointerFwd.h"
+
+class SImage;
+enum EOrientation : int;
+struct FSlateBrush;
+struct FSlateColor;
+
+namespace SlateIM
+{
+	void BeginContainer();
+	void EndContainer();
+
+	void BeginStack(EOrientation Orientation);
+	void EndStack();
+
+	void BeginWrap(EOrientation Orientation);
+	void EndWrap();
+
+	void Image_Internal(const FSlateColor& ColorAndOpacity, FVector2D DesiredSize, TFunctionRef<const FSlateBrush* (const TSharedRef<SImage>&)> GetBrush);
+} // SlateIM
diff --git a/Engine/Plugins/Experimental/SlateIM/Source/SlateIM/Public/SlateIM.h b/Engine/Plugins/Experimental/SlateIM/Source/SlateIM/Public/SlateIM.h
--- a/Engine/Plugins/Experimental/SlateIM/Source/SlateIM/Public/SlateIM.h
+++ b/Engine/Plugins/Experimental/SlateIM/Source/SlateIM/Public/SlateIM.h
@@ -5,19 +5,28 @@
 #include "Layout/Margin.h"
 #include "Modules/ModuleInterface.h"
 #include "Styling/SlateColor.h"
+#include "Styling/StyleColors.h"
 #include "Textures/SlateIcon.h"
 #include "Types/SlateEnums.h"
 #include "Widgets/Layout/Anchors.h"
+#if WITH_ENGINE
+#include "Engine/EngineTypes.h"
+#endif
 
 #if WITH_EDITOR
 class IAssetViewport;
 #endif
 #if WITH_ENGINE
+class FCanvasItem;
+class UCanvas;
+class UFont;
 class UGameViewportClient;
 class ULocalPlayer;
 class UMaterialInterface;
 class UTextureRenderTarget2D;
+class UTexture;
 class UTexture2D;
+struct FCanvasIcon;
 #endif
 
 struct FButtonStyle;
@@ -860,5 +869,277 @@ namespace SlateIM
 	 */
 	SLATEIM_API float GetKeyAnalogValue(const FKey& InKey);
 #pragma endregion Inputs
+
+#pragma region Canvas
+#if WITH_ENGINE
+	namespace Canvas
+	{
+		enum class ECanvasUpdateType : uint8
+		{
+			/** Redraw the canvas every frame. */
+			EveryFrame,
+
+			/**
+			 * Only redraw when specifically invalidated.
+			 * Warning: Textures and materials may not be rendered correctly if they are not loaded or shaders are not cached.
+			 */
+			Invalidation
+
+			/** TODO: Automatically redraw when the hash of the draw queue changes. */
+			// HashChange
+		};
+
+		struct FCanvasParams
+		{
+			FSlateColor ColorAndOpacity = FSlateColor(EStyleColor::White);
+			FVector2D DesiredSize = FVector2D::ZeroVector;
+			UObject* WorldContext = nullptr;
+			ECanvasUpdateType UpdateType = ECanvasUpdateType::EveryFrame;
+		};
+
+		/**
+		 * Begin queueing commands to draw to a canvas render target.
+		 * 
+		 * @param Width The width of the canvas in pixels.
+		 * @param Height The height of the canvas in pixels.
+		 * @param ColorAndOpacity (optional) Tint applied to the resulting image.
+		 * @param DesiredSize (optional) The size of the image drawn in the UI. This may be different from the size of the canvas.
+		 * @param WorldContext (optional) The world related to the canvas being drawn, if any.
+		 */
+		void BeginCanvas(int32 Width, int32 Height, const FCanvasParams& CanvasParams = FCanvasParams());
+
+		/**
+		 * Ends the currently drawn canvas and executes all queued commands.
+		 */
+		void EndCanvas();
+
+		/**
+		 * Invalidate the current canvas, forcing a redraw if its UpdateType is set to Invalidation.
+		 */
+		void Invalidate();
+
+		/** Sets the position of the lower-right corner of the clipping region of the canvas */
+		void SetClip(const FVector2f& ClipPosition);
+
+		/** Sets the draw color of the canvas. */
+		void SetDrawColor(const FLinearColor& Color);
+		void SetDrawColor(const FColor& Color);
+
+		/**
+		 * Queue drawing a canvas item to the canvas.
+		 * 
+		 * @param Item The canvas item to draw. The item *must* stay in scope until EndCanvas().
+		 * @param CanvasPosition (optional) The position to draw the item on the canvas.
+		 */
+		void DrawItem(TSharedRef<FCanvasItem> ItemPtr, const FVector2f& CanvasPosition = FVector2f::ZeroVector);
+
+		/**
+		 * Queue drawing a canvas icon (texture) to the canvas.
+		 * 
+		 * @param Icon The icon to draw. The texture draw area is in pixels.
+		 * @param CanvasPosition (optional) The position of the icon on the canvas, in pixels.
+		 * @param Scale (optional) The scale of the icon.
+		 */
+		void DrawIcon(const FCanvasIcon& Icon, const FVector2f& CanvasPosition = FVector2f::ZeroVector, const FVector2f& Scale = FVector2f::UnitVector);
+
+		/**
+		 * Queue drawing a line to the canvas.
+		 * 
+		 * @param CanvasPositionA One end of the line on the canvas, in pixels.
+		 * @param CanvasPositionB The other end of the line on the canvas, in pixels.
+		 * @param Thickness (optional) The thickness of the line in pixels.
+		 * @param RenderColor (optional) The color of the line.
+		 */
+		void DrawLine(const FVector2D& CanvasPositionA, const FVector2D& CanvasPositionB, float Thickness = 1.0f, FLinearColor RenderColor = FLinearColor::White);
+
+		/** 
+		 * Extra optional parameters for rendering tiles (textures and materials.)
+		 */
+		struct FTileRenderParams
+		{
+			// The top left of the area on the tile to draw, in texels.
+			FVector2D UVPosition = FVector2D::ZeroVector;
+
+			// The size of the area on the tile to draw, in texels.
+			FVector2D UVSize = FVector2D::UnitVector;
+
+			// A tint to apply to tile.
+			FLinearColor RenderColor = FLinearColor::White;
+
+			// The blend mode with which to draw the tile.
+			EBlendMode BlendMode = BLEND_Translucent;
+
+			// The rotation of the tile, in degrees.
+			float Rotation = 0.f;
+
+			// The point on the tile around which it is rotated, in texels.
+			FVector2D PivotPoint = FVector2D(0.5, 0.5);
+		};
+
+		/**
+		 * Queue drawing a texture to the canvas.
+		 *
+		 * @param RenderTexture The texture to draw.
+		 * @param CanvasPosition The position of the texture on the canvas, in pixels.
+		 * @param CanvasSize The size of the texture on the canvas, in pixels.
+		 * @param RenderParams (optional) Extra optional parameters for tile drawing.
+		 */
+		void DrawTexture(UTexture* RenderTexture, const FVector2D& CanvasPosition, const FVector2D& CanvasSize, const FTileRenderParams& RenderParams = FTileRenderParams());
+
+		/**
+		 * Queue drawing a material to the canvas.
+		 *
+		 * @param RenderMaterial The material to draw.
+		 * @param CanvasPosition The position of the material on the canvas, in pixels.
+		 * @param CanvasSize The size of the material on the canvas, in pixels.
+		 * @param RenderParams (optional) Extra optional parameters for tile drawing. RenderColor and BlendMode are not used.
+		 */
+		void DrawMaterial(UMaterialInterface* RenderMaterial, const FVector2D& CanvasPosition, const FVector2D& CanvasSize, const FTileRenderParams& RenderParams = FTileRenderParams());
+
+		/**
+		 * Extra optional parameters for rendering text.
+		 */
+		struct FTextRenderParams
+		{
+			// The scale of the text on the canvas.
+			FVector2D Scale = FVector2D(1.0f, 1.0f);
+
+			// The color of the text.
+			FLinearColor RenderColor = FLinearColor::White;
+			
+			// Adjust the distance between letters.
+			float Kerning = 0.0f;
+
+			// The color of the text shadow.
+			FLinearColor ShadowColor = FLinearColor::Black;
+
+			// The offset of the shadow from the text it is shadowing.
+			FVector2D ShadowOffset = FVector2D::UnitVector;
+
+			// If true, the text is centered on the given draw position on the X axis.
+			bool bCentreX = false;
+
+			// If true, the text is centered on the given draw position on the Y axis.
+			bool bCentreY = false;
+
+			// If true, outline the text.
+			bool bOutlined = false;
+
+			// Color of the text outline.
+			FLinearColor OutlineColor = FLinearColor::Black;
+		};
+
+		/**
+		 * Queue drawing text to the canvas.
+		 *
+		 * @param RenderFont The font used to render the text.
+		 * @param RenderText The string to render.
+		 * @param CanvasPosition The position of the text on the canvas, in pixels.
+		 * @param RenderParams Extra optional parameters for rendering the text.
+		 */
+		void DrawText(UFont* RenderFont, const FString& RenderText, const FVector2D& CanvasPosition, const FTextRenderParams& RenderParams);
+
+		/**
+		 * Queue drawing text to the canvas.
+		 * 
+		 * Alternate method allowing the use of FFontRenderInfo.
+		 *
+		 * @param RenderFont The font used to render the text.
+		 * @param RenderText The string to render.
+		 * @param CanvasPosition The position of the text on the canvas, in pixels.
+		 * @param Scale (optional) The scale of the text on the canvas.
+		 * @param RenderInfo (optional) Extra optional parameters for rendering the text.
+		 */
+		void DrawText(UFont* RenderFont, const FString& RenderText, const FVector2f& CanvasPosition, const FVector2f& Scale = FVector2f::UnitVector, const FFontRenderInfo& RenderInfo = FFontRenderInfo());
+
+		/**
+		 * Queue drawing wrapped text to canvas.
+		 * 
+		 * Uses FFontRenderInfo for text effets.
+		 * 
+		 * Text size / wrapping is automatically computed based on the position of the text and the @see SetCanvasClip.
+		 * The default clipping position is the entire canvas. It does not reset between queued calls.
+		 * 
+		 * @param RenderFont The font used to render the text.
+		 * @param RenderText The text to render.
+		 * @param CanvasPosition The position of the text on the canvas, in pixels.
+		 * @param bCenterTextX If true, the text is centered on the given draw position on the X axis.
+		 * @param bCenterTextY If true, the text is centered on the given draw position on the Y axis.
+		 * @param Scale (optional) The scale of the text on the canvas.
+		 * @param RenderInfo (optional) Extra optional parameters for rendering the text.
+		 */
+		void DrawWrappedText(const UFont* RenderFont, const FString& RenderText, const FVector2f& CanvasPosition, bool bCenterTextX, bool bCenterTextY, const FVector2f& Scale = FVector2f::UnitVector, const FFontRenderInfo& RenderInfo = FFontRenderInfo());
+
+		/**
+		 * Extra optional parameters for rendering borders.
+		 * 
+		 * Tile refers to the centre texture.
+		 */
+		struct FBorderRenderParams : FTileRenderParams
+		{
+			// The scale of the border.
+			FVector2D BorderScale = FVector2D(0.1, 0.1);
+
+			// The scale of the background.
+			FVector2D BackgroundScale = FVector2D(0.1, 0.1);
+
+			// Frame corner size, in percent (should be < 0.5f).
+			FVector2D CornerSize = FVector2D::ZeroVector;
+		};
+
+		/**
+		 * Queue drawing a border to the canvas.
+		 * 
+		 * @param BorderTexture The texture for the center of the border.
+		 * @param BackgroundTexture The texture for the background of the border.
+		 * @param LeftBorderTexture The texture for the left of the border.
+		 * @param RightBorderTexture The texture for the right of the border.
+		 * @param TopBorderTexture The texture for the top of the border.
+		 * @param BottomBorderTexture The texture for the bottom of the border.
+		 * @param CanvasPosition The position of the border on the canvas, in pixels.
+		 * @param CanvasSize The size of the border on the canvas, in pixels.
+		 * @param RenderParams (optional) Extra optional params for rendering the border.
+		 */
+		void DrawBorder(UTexture* BorderTexture, UTexture* BackgroundTexture, UTexture* LeftBorderTexture, UTexture* RightBorderTexture, UTexture* TopBorderTexture, UTexture* BottomBorderTexture, const FVector2D& CanvasPosition, const FVector2D& CanvasSize, const FBorderRenderParams& RenderParams = FBorderRenderParams());
+
+		/**
+		 * Queue drawing an unfilled box to the canvas.
+		 * 
+		 * @param CanvasPosition The position of the box on the canvas, in pixels.
+		 * @param CanvasSize The size of the box on the canvas, in pixels.
+		 * @param Thickness (optional) The thicknses of the lines to draw.
+		 * @param RenderColor (optional) The color of the lines to draw.
+		 */
+		void DrawBox(const FVector2D& CanvasPosition, const FVector2D& CanvasSize, float Thickness = 1.0f, FLinearColor RenderColor = FLinearColor::White);
+
+		/**
+		 * Queue drawing a series of texture-filled triangles to the canvas.
+		 * 
+		 * @param RenderTexture The texture to draw with.
+		 * @param Triangles The list of triangles to draw.
+		 */
+		void DrawTriangles(UTexture* RenderTexture, TSharedRef<TArray<FCanvasUVTri>> Triangles);
+
+		/**
+		 * Queue drawing a series of material-filled triangles to the canvas.
+		 *
+		 * @param RenderTexture The texture to draw with.
+		 * @param Triangles The list of triangles to draw.
+		 */
+		void DrawMaterialTriangles(UMaterialInterface* RenderMaterial, TSharedRef<TArray<FCanvasUVTri>> Triangles);
+
+		/**
+		 * Queue drawing a filled polygon to the canvas.
+		 * 
+		 * @param RenderTexture The texture to draw with.
+		 * @param CanvasPosition The position of the center of the polygon on the canvas, in pixels.
+		 * @param Radius The radius of the polygon, in pixels.
+		 * @param NumberOfSides The number of sides of the polygon.
+		 * @param RenderColor (optional) The color of the polygon.
+		 */
+		void DrawPolygon(UTexture* RenderTexture, const FVector2D& CanvasPosition, const FVector2D& Radius, int32 NumberOfSides, FLinearColor RenderColor = FLinearColor::White);
+	}
+#endif // WITH_ENGINE
+#pragma endregion Canvas
 }
 
