Use a uint instead of a FName for expanded states

#rnx
#Synced-CL 49846476

diff --git a/Engine/Plugins/Experimental/SlateIM/Source/SlateIM/Private/Containers/SImTableView.cpp b/Engine/Plugins/Experimental/SlateIM/Source/SlateIM/Private/Containers/SImTableView.cpp
--- a/Engine/Plugins/Experimental/SlateIM/Source/SlateIM/Private/Containers/SImTableView.cpp
+++ b/Engine/Plugins/Experimental/SlateIM/Source/SlateIM/Private/Containers/SImTableView.cpp
@@ -400,7 +400,7 @@ const TArray<TSharedRef<FSlateIMTableRow>>& FSlateIMTableBody::GetTableRows() co
 	return Children;
 }
 
-bool FSlateIMTableBody::GetSavedExpansionState(FName InRowId)
+bool FSlateIMTableBody::GetSavedExpansionState(uint32 InRowId) const
 {
 	return ExpandedItems.Contains(InRowId);
 }
@@ -412,9 +412,9 @@ void FSlateIMTableBody::HandleOnExpansionChanged(TSharedRef<FSlateIMTableRow> In
 		return;
 	}
 
-	const FName RowId = InItem->GetRowId();
+	const uint32 RowId = InItem->GetRowId();
 
-	if (RowId.IsNone())
+	if (RowId == 0)
 	{
 		return;
 	}
diff --git a/Engine/Plugins/Experimental/SlateIM/Source/SlateIM/Private/Containers/SImTableView.h b/Engine/Plugins/Experimental/SlateIM/Source/SlateIM/Private/Containers/SImTableView.h
--- a/Engine/Plugins/Experimental/SlateIM/Source/SlateIM/Private/Containers/SImTableView.h
+++ b/Engine/Plugins/Experimental/SlateIM/Source/SlateIM/Private/Containers/SImTableView.h
@@ -71,8 +71,8 @@ public:
 	virtual void UpdateChild(FSlateIMChild Child, int32 Index, const FSlateIMSlotData& AlignmentData) override;
 	virtual TSharedPtr<SWidget> GetAsWidget() override;
 
-	FName GetRowId() const { return RowId; }
-	void SetRowId(FName InRowId) { RowId = InRowId; }
+	uint32 GetRowId() const { return RowId; }
+	void SetRowId(uint32 InRowId) { RowId = InRowId; }
 
 	void GetChildRows(TArray<TSharedRef<FSlateIMTableRow>>& OutRows) const;
 
@@ -94,7 +94,7 @@ public:
 
 private:
 	// A mix of cell widgets and child rows
-	FName RowId;
+	uint32 RowId = 0;
 	TArray<FSlateIMChild> Children;
 	int32 ColumnCount = 0;
 	TSharedPtr<SImTableView> OwningTable;
@@ -142,14 +142,14 @@ public:
 
 	const TArray<TSharedRef<FSlateIMTableRow>>& GetTableRows() const;
 
-	bool GetSavedExpansionState(FName InRowId);
+	bool GetSavedExpansionState(uint32 InRowId) const;
 
 	void HandleOnExpansionChanged(TSharedRef<FSlateIMTableRow> InItem, bool bInExpanded);
 
 private:
 	TArray<TSharedRef<FSlateIMTableRow>> Children;
 	TSharedPtr<SImTableView> OwningTable;
-	TSet<FName> ExpandedItems;
+	TSet<uint32> ExpandedItems;
 };
 
 // No direct widget representation
diff --git a/Engine/Plugins/Experimental/SlateIM/Source/SlateIM/Private/Implementation/SlateIM_Containers.cpp b/Engine/Plugins/Experimental/SlateIM/Source/SlateIM/Private/Implementation/SlateIM_Containers.cpp
--- a/Engine/Plugins/Experimental/SlateIM/Source/SlateIM/Private/Implementation/SlateIM_Containers.cpp
+++ b/Engine/Plugins/Experimental/SlateIM/Source/SlateIM/Private/Implementation/SlateIM_Containers.cpp
@@ -613,7 +613,7 @@ namespace SlateIM
 		return Row && Row->AreTableRowContentsRequired();
 	}
 
-	bool BeginTableRowChildren(FName ParentRowId)
+	bool BeginTableRowChildren(uint32 ParentRowId)
 	{
 		// TODO - Let caller set default expansion state
 		// TODO - Save and restore expansion state
@@ -634,7 +634,7 @@ namespace SlateIM
 
 		SlateIM::BeginTableRow();
 
-		if (!ParentRowId.IsNone())
+		if (ParentRowId != 0)
 		{
 			if (TSharedPtr<SImTableView> OwningTable = ParentRow->GetOwningTable())
 			{
diff --git a/Engine/Plugins/Experimental/SlateIM/Source/SlateIM/Private/Misc/SlateIMExamples.cpp b/Engine/Plugins/Experimental/SlateIM/Source/SlateIM/Private/Misc/SlateIMExamples.cpp
--- a/Engine/Plugins/Experimental/SlateIM/Source/SlateIM/Private/Misc/SlateIMExamples.cpp
+++ b/Engine/Plugins/Experimental/SlateIM/Source/SlateIM/Private/Misc/SlateIMExamples.cpp
@@ -688,7 +688,8 @@ void FSlateIMTestWidget::DrawTables()
 							SlateIM::VAlign(VAlign_Center);
 							SlateIM::Text(TEXT("Country"));
 						}
-						if (SlateIM::BeginTableRowChildren("CanadaChildren"))
+						static constexpr FLazyName CanadaChildren = TEXT("CanadaChildren");
+						if (SlateIM::BeginTableRowChildren(GetTypeHash(CanadaChildren)))
 						{
 							if (SlateIM::NextTableCell())
 							{
@@ -753,7 +754,8 @@ void FSlateIMTestWidget::DrawTables()
 						SlateIM::VAlign(VAlign_Center);
 						SlateIM::Text(TEXT("Country"));
 					}
-					if (SlateIM::BeginTableRowChildren("UnitedStatesChildren"))
+					static const uint32 RowID = FCrc::StrCrc32("UnitedStatesChildren");
+					if (SlateIM::BeginTableRowChildren(RowID))
 					{
 						if (SlateIM::NextTableCell())
 						{
diff --git a/Engine/Plugins/Experimental/SlateIM/Source/SlateIM/Public/SlateIM.h b/Engine/Plugins/Experimental/SlateIM/Source/SlateIM/Public/SlateIM.h
--- a/Engine/Plugins/Experimental/SlateIM/Source/SlateIM/Public/SlateIM.h
+++ b/Engine/Plugins/Experimental/SlateIM/Source/SlateIM/Public/SlateIM.h
@@ -262,11 +262,11 @@ namespace SlateIM
 	 * In the case of a dynamic table, providing a parent row id ensures that its expansion state is maintained correctly.
 	 * A parent row id is not required if the table is not dynamic.
 	 * 
-	 * @param Id of the parent row.
+	 * @param ParentRowId Id of the parent row.
 	 * 
 	 * @return Whether the parent row is expanded or not. When false, the child content can be skipped to save cycles.
 	 */
-	SLATEIM_API bool BeginTableRowChildren(FName ParentRowId = NAME_None);
+	SLATEIM_API bool BeginTableRowChildren(uint32 ParentRowId = 0);
 	/**
 	 * Stop adding child content to the parent row. This must be called no matter the result of BeginTableRowChildren()
 	 */
