SlateIMExample: Use a constexpr ID

#rnx

#Synced-CL 49889763
#Synced-CL 49889763

diff --git a/Engine/Plugins/Experimental/SlateIM/Source/SlateIM/Private/Misc/SlateIMExamples.cpp b/Engine/Plugins/Experimental/SlateIM/Source/SlateIM/Private/Misc/SlateIMExamples.cpp
--- a/Engine/Plugins/Experimental/SlateIM/Source/SlateIM/Private/Misc/SlateIMExamples.cpp
+++ b/Engine/Plugins/Experimental/SlateIM/Source/SlateIM/Private/Misc/SlateIMExamples.cpp
@@ -754,7 +754,7 @@ void FSlateIMTestWidget::DrawTables()
 						SlateIM::VAlign(VAlign_Center);
 						SlateIM::Text(TEXT("Country"));
 					}
-					static const uint32 RowID = FCrc::StrCrc32("UnitedStatesChildren");
+					constexpr uint32 RowID = UE::HashStringFNV1a32("UnitedStatesChildren");
 					if (SlateIM::BeginTableRowChildren(RowID))
 					{
 						if (SlateIM::NextTableCell())
