SlateIM: Tables now support selection settings.
- NextTableCell() now includes an out parameter to say whether the current row is selected.

#Synced-CL 49990324

diff --git a/Engine/Plugins/Experimental/SlateIM/Source/SlateIM/Private/Implementation/SlateIM_Containers.cpp b/Engine/Plugins/Experimental/SlateIM/Source/SlateIM/Private/Implementation/SlateIM_Containers.cpp
--- a/Engine/Plugins/Experimental/SlateIM/Source/SlateIM/Private/Implementation/SlateIM_Containers.cpp
+++ b/Engine/Plugins/Experimental/SlateIM/Source/SlateIM/Private/Implementation/SlateIM_Containers.cpp
@@ -321,8 +321,13 @@ namespace SlateIM
 	{
 		SlateIM::EndHorizontalStack();
 	}
-	
+
 	void BeginTable(const FTableViewStyle* InStyle, const FTableRowStyle* InRowStyle)
+	{
+		BeginTable({ .Style = InStyle, .RowStyle = InRowStyle });
+	}
+	
+	void BeginTable(const FTableParams& Params)
 	{
 		TSharedPtr<SImTableView> ContainerWidget;
 
@@ -331,21 +336,21 @@ namespace SlateIM
 			FWidgetScope<SImTableView> Scope;
 			ContainerWidget = Scope.GetWidget();
 			
-			Scope.HashData(InStyle);
+			Scope.HashData(Params.Style);
 
 			if (!ContainerWidget)
 			{
 				ContainerWidget = SNew(SImTableView)
-					.TreeViewStyle(InStyle)
-					.SelectionMode(ESelectionMode::None);
+					.TreeViewStyle(Params.Style)
+					.SelectionMode(Params.SelectionMode);
 				Scope.UpdateWidget(ContainerWidget);
 			}
 			else if (Scope.IsDataHashDirty())
 			{
-				ContainerWidget->SetStyle(InStyle);
+				ContainerWidget->SetStyle(Params.Style);
 			}
 
-			ContainerWidget->SetTableRowStyle(InRowStyle);
+			ContainerWidget->SetTableRowStyle(Params.RowStyle);
 			ContainerWidget->BeginTableUpdates();
 		}
 
@@ -550,7 +555,7 @@ namespace SlateIM
 		FSlateIMManager::Get().PopContainer<FSlateIMTableBody>();
 	}
 
-	bool NextTableCell()
+	bool NextTableCell(bool* bOutRowSelected)
 	{
 		SCOPED_NAMED_EVENT_TEXT("SlateIM::NextTableCell", FColorList::Goldenrod);
 
@@ -598,7 +603,14 @@ namespace SlateIM
 				Row = FSlateIMManager::Get().GetCurrentContainer<FSlateIMTableRow>();
 				ensure(Row);
 			}
-			
+
+			if (bOutRowSelected)
+			{
+				if (TSharedPtr<SImTableView> OwningTable = Row->GetOwningTable())
+				{
+					*bOutRowSelected = OwningTable->IsItemSelected(Row.ToSharedRef());
+				}
+			}			
 
 			// Default table cells to fill (NextHAlign might have been set by the user, so we check first) 
 			if (!FSlateIMManager::Get().NextHAlign.IsSet())
diff --git a/Engine/Plugins/Experimental/SlateIM/Source/SlateIM/Public/SlateIM.h b/Engine/Plugins/Experimental/SlateIM/Source/SlateIM/Public/SlateIM.h
--- a/Engine/Plugins/Experimental/SlateIM/Source/SlateIM/Public/SlateIM.h
+++ b/Engine/Plugins/Experimental/SlateIM/Source/SlateIM/Public/SlateIM.h
@@ -2,6 +2,7 @@
 
 #pragma once
 
+#include "Framework/Views/ITypedTableView.h"
 #include "Layout/Margin.h"
 #include "Modules/ModuleInterface.h"
 #include "Styling/SlateColor.h"
@@ -184,6 +185,26 @@ namespace SlateIM
 	SLATEIM_API void BeginVerticalWrap();
 	SLATEIM_API void EndVerticalWrap();
 
+	struct FTableParams
+	{
+		// A Table View Style to override the default style. Changes to the style after the first invocation are not reflected,
+		// provide a new style object to update the visuals at runtime
+		const FTableViewStyle* Style = nullptr;
+		
+		// A Table Row Style to override the default table row style. Changing the style only affects new table rows.
+		const FTableRowStyle* RowStyle = nullptr;
+
+		// Select mode for the table.
+		ESelectionMode::Type SelectionMode = ESelectionMode::None;
+	};
+
+	/**
+	 * Begin a Table container. The contents are automatically scrollable if the table is sized smaller than all its content.
+	 * 
+	 * @param Params Settings for table construction. Not optional for compatibility with the below method.
+	 */
+	SLATEIM_API void BeginTable(const FTableParams& Params);
+
 	/**
 	 * Begin a Table container. The contents are automatically scrollable if the table is sized smaller than all its content.
 	 * 
@@ -192,6 +213,7 @@ namespace SlateIM
 	 * @param InRowStyle (optional) A Table Row Style to override the default table row style. Changing the style only affects new table rows.
 	 */
 	SLATEIM_API void BeginTable(const FTableViewStyle* InStyle = nullptr, const FTableRowStyle* InRowStyle = nullptr);
+
 	/**
 	 * End the current Table container.
 	 */
@@ -251,11 +273,13 @@ namespace SlateIM
 	/**
 	 * Begin the content for the next table cell. The column and row are tracked internally based on the number of columns.
 	 * 
+	 * @param bOutRowSelected (optional) If supplied, will return true if the row this cell belongs to is selected.
+	 * 
 	 * @return Whether the contents of this table cell are visible. When false, the content can be skipped to save cycles.
 	 *
 	 * @note There is no "End" function counterpart to this call. This function supports the various slot functions (Padding, HAlign, VAlign, Min/Max Width/Height, etc).
 	 */
-	SLATEIM_API bool NextTableCell();
+	SLATEIM_API bool NextTableCell(bool* bOutRowSelected = nullptr);
 
 	/**
 	 * Begin adding content as a child to the current table row. NextTableCell() must still be called after this function.
