SlateIM: Changed a table in the examples to be selectable and adds an indication of the selected row.

#Synced-CL 49990324

diff --git a/Engine/Plugins/Experimental/SlateIM/Source/SlateIM/Private/Misc/SlateIMExamples.cpp b/Engine/Plugins/Experimental/SlateIM/Source/SlateIM/Private/Misc/SlateIMExamples.cpp
--- a/Engine/Plugins/Experimental/SlateIM/Source/SlateIM/Private/Misc/SlateIMExamples.cpp
+++ b/Engine/Plugins/Experimental/SlateIM/Source/SlateIM/Private/Misc/SlateIMExamples.cpp
@@ -596,7 +596,7 @@ void FSlateIMTestWidget::DrawTables()
 			SlateIM::Fill();
 			SlateIM::HAlign(HAlign_Fill);
 			SlateIM::MaxHeight(200.f);
-			SlateIM::BeginTable();
+			SlateIM::BeginTable({ .SelectionMode = ESelectionMode::Single });
 			SlateIM::BeginTableHeader();
 			SlateIM::NextTableColumn(TEXT("Item"));
 			// Makes clickable area fill the column header
@@ -615,12 +615,22 @@ void FSlateIMTestWidget::DrawTables()
 			SlateIM::BeginTableBody();
 			for (int32 i = 0; i < NumTableItems; ++i)
 			{
-				if (SlateIM::NextTableCell())
+				bool bIsSelected = false;
+
+				if (SlateIM::NextTableCell(&bIsSelected))
 				{
 					SlateIM::Padding(FMargin(5, 0));
 					SlateIM::Fill();
 					SlateIM::VAlign(VAlign_Center);
-					SlateIM::Text(FString::Printf(TEXT("Item %d/%d"), i + 1, NumTableItems), FColor::MakeRedToGreenColorFromScalar(static_cast<float>(i) / NumTableItems));
+
+					if (!bIsSelected)
+					{
+						SlateIM::Text(FString::Printf(TEXT("Item %d/%d"), i + 1, NumTableItems), FColor::MakeRedToGreenColorFromScalar(static_cast<float>(i) / NumTableItems));
+					}
+					else
+					{
+						SlateIM::Text(FString::Printf(TEXT("Item %d/%d [Selected]"), i + 1, NumTableItems), FColor::MakeRedToGreenColorFromScalar(static_cast<float>(i) / NumTableItems));
+					}
 				}
 
 				if (SlateIM::NextTableCell())
