SlateIM: Table Row IDs are now unique per parent.
- Default expansion status can now be set for rows with consistent Ids.

#Synced-CL 49990324

diff --git a/Engine/Plugins/Experimental/SlateIM/Source/SlateIM/Private/Containers/SImTableView.cpp b/Engine/Plugins/Experimental/SlateIM/Source/SlateIM/Private/Containers/SImTableView.cpp
--- a/Engine/Plugins/Experimental/SlateIM/Source/SlateIM/Private/Containers/SImTableView.cpp
+++ b/Engine/Plugins/Experimental/SlateIM/Source/SlateIM/Private/Containers/SImTableView.cpp
@@ -12,6 +12,7 @@ SLATE_IMPLEMENT_WIDGET(SImTableView)
 namespace SlateIM::Private
 {
 	const FName ExpanderColumnName = TEXT("Expander");
+	const uint32 NoRowId = 0;
 }
 
 void SImTableHeader::PrivateRegisterAttributes(FSlateAttributeInitializer&)
@@ -76,6 +77,37 @@ TSharedPtr<SImTableView> FSlateIMTableHeader::GetOwningTable() const
 	return OwningTable;
 }
 
+bool FSlateIMRowExpansionStateData::HasSavedExpansionState(uint32 InRowId) const
+{
+	if (InRowId != SlateIM::Private::NoRowId)
+	{
+		return ExpansionStates.Contains(InRowId);
+	}
+
+	return false;
+}
+
+bool FSlateIMRowExpansionStateData::GetSavedExpansionState(uint32 InRowId, bool bDefaultState) const
+{
+	if (InRowId != SlateIM::Private::NoRowId)
+	{
+		if (const bool* ExpandedState = ExpansionStates.Find(InRowId))
+		{
+			return *ExpandedState;
+		}
+	}
+
+	return bDefaultState;
+}
+
+void FSlateIMRowExpansionStateData::SetSavedExpansionState(uint32 InRowId, bool bState)
+{
+	if (InRowId != SlateIM::Private::NoRowId)
+	{
+		ExpansionStates.FindOrAdd(InRowId) = bState;
+	}
+}
+
 void FSlateIMTableHeader::SetOwningTable(const TSharedPtr<SImTableView>& InOwningTable)
 {
 	OwningTable = InOwningTable;
@@ -141,6 +173,16 @@ void FSlateIMTableRow::GetChildRows(TArray<TSharedRef<FSlateIMTableRow>>& OutRow
 	}
 }
 
+TSharedPtr<ISlateIMContainer> FSlateIMTableRow::GetParent() const
+{
+	return Parent.Pin();
+}
+
+void FSlateIMTableRow::SetParent(TSharedPtr<ISlateIMContainer> InParent)
+{
+	Parent = InParent;
+}
+
 int32 FSlateIMTableRow::CountCellWidgetsUpToIndex(int32 Index) const
 {
 	int32 CellCount = 0;
@@ -400,35 +442,6 @@ const TArray<TSharedRef<FSlateIMTableRow>>& FSlateIMTableBody::GetTableRows() co
 	return Children;
 }
 
-bool FSlateIMTableBody::GetSavedExpansionState(uint32 InRowId) const
-{
-	return ExpandedItems.Contains(InRowId);
-}
-
-void FSlateIMTableBody::HandleOnExpansionChanged(TSharedRef<FSlateIMTableRow> InItem, bool bInExpanded)
-{
-	if (!OwningTable.IsValid())
-	{
-		return;
-	}
-
-	const uint32 RowId = InItem->GetRowId();
-
-	if (RowId == 0)
-	{
-		return;
-	}
-
-	if (bInExpanded)
-	{
-		ExpandedItems.Add(RowId);
-	}
-	else
-	{
-		ExpandedItems.Remove(RowId);
-	}
-}
-
 void SImTableView::PrivateRegisterAttributes(FSlateAttributeInitializer&)
 {
 }
@@ -705,7 +718,17 @@ void SImTableView::AddColumn_Internal(const FName& ColumnId, const FStringView&
 
 void SImTableView::HandleOnExpansionChanged(TSharedRef<FSlateIMTableRow> InItem, bool bInExpanded)
 {
-	Body->HandleOnExpansionChanged(InItem, bInExpanded);
+	if (TSharedPtr<ISlateIMContainer> Parent = InItem->GetParent())
+	{
+		if (Parent->IsA<FSlateIMTableBody>())
+		{
+			StaticCastSharedPtr<FSlateIMTableBody>(Parent)->SetSavedExpansionState(InItem->GetRowId(), bInExpanded);
+		}
+		else
+		{
+			StaticCastSharedPtr<FSlateIMTableRow>(Parent)->SetSavedExpansionState(InItem->GetRowId(), bInExpanded);
+		}
+	}
 }
 
 void SImTableHeader::Construct(const SHeaderRow::FArguments& InArgs, const TSharedRef<FSlateIMTableHeader> InHeader)
diff --git a/Engine/Plugins/Experimental/SlateIM/Source/SlateIM/Private/Containers/SImTableView.h b/Engine/Plugins/Experimental/SlateIM/Source/SlateIM/Private/Containers/SImTableView.h
--- a/Engine/Plugins/Experimental/SlateIM/Source/SlateIM/Private/Containers/SImTableView.h
+++ b/Engine/Plugins/Experimental/SlateIM/Source/SlateIM/Private/Containers/SImTableView.h
@@ -2,7 +2,7 @@
 
 #pragma once
 
-#include "Containers/Set.h"
+#include "Containers/Map.h"
 #include "ISlateIMContainer.h"
 #include "Misc/ISlateIMChild.h"
 #include "Widgets/Views/IItemsSource.h"
@@ -52,11 +52,21 @@ private:
 };
 
 
+class FSlateIMRowExpansionStateData
+{
+public:
+	bool HasSavedExpansionState(uint32 InRowId) const;
+	bool GetSavedExpansionState(uint32 InRowId, bool bDefaultState = true) const;
+	void SetSavedExpansionState(uint32 InRowId, bool bState);
+
+private:
+	TMap<uint32, bool> ExpansionStates;
+};
 
 /**
  * Row
  */
-class FSlateIMTableRow : public ISlateIMContainer, public ISlateIMChild, public TSharedFromThis<FSlateIMTableRow>
+class FSlateIMTableRow : public ISlateIMContainer, public ISlateIMChild, public FSlateIMRowExpansionStateData, public TSharedFromThis<FSlateIMTableRow>
 {
 	SLATE_IM_TYPE_DATA(FSlateIMTableRow, ISlateIMChild)
 	// Deadly diamond of death, but it's okay since ISlateIMContainer and ISlateIMChild don't have implementations for ISlateIMTypeChecking
@@ -76,6 +86,9 @@ public:
 
 	void GetChildRows(TArray<TSharedRef<FSlateIMTableRow>>& OutRows) const;
 
+	TSharedPtr<ISlateIMContainer> GetParent() const;
+	void SetParent(TSharedPtr<ISlateIMContainer> InParent);
+
 	int32 CountCellWidgetsUpToIndex(int32 Index) const;
 	int32 FindColumnIndex(const FName& ColumnName) const;
 
@@ -93,7 +106,7 @@ public:
 	bool AreTableRowContentsRequired();
 
 private:
-	// A mix of cell widgets and child rows
+	TWeakPtr<ISlateIMContainer> Parent;
 	uint32 RowId = 0;
 	TArray<FSlateIMChild> Children;
 	int32 ColumnCount = 0;
@@ -118,7 +131,7 @@ private:
 /**
  * Body
  */
-class FSlateIMTableBody : public ISlateIMContainer, public ISlateIMChild, public TSharedFromThis<FSlateIMTableBody>
+class FSlateIMTableBody : public ISlateIMContainer, public ISlateIMChild, public FSlateIMRowExpansionStateData, public TSharedFromThis<FSlateIMTableBody>
 {
 	SLATE_IM_TYPE_DATA(FSlateIMTableBody, ISlateIMChild)
 	// Deadly diamond of death, but it's okay since ISlateIMContainer and ISlateIMChild don't have implementations for ISlateIMTypeChecking
@@ -142,14 +155,9 @@ public:
 
 	const TArray<TSharedRef<FSlateIMTableRow>>& GetTableRows() const;
 
-	bool GetSavedExpansionState(uint32 InRowId) const;
-
-	void HandleOnExpansionChanged(TSharedRef<FSlateIMTableRow> InItem, bool bInExpanded);
-
 private:
 	TArray<TSharedRef<FSlateIMTableRow>> Children;
 	TSharedPtr<SImTableView> OwningTable;
-	TSet<uint32> ExpandedItems;
 };
 
 // No direct widget representation
diff --git a/Engine/Plugins/Experimental/SlateIM/Source/SlateIM/Private/Implementation/SlateIM_Containers.cpp b/Engine/Plugins/Experimental/SlateIM/Source/SlateIM/Private/Implementation/SlateIM_Containers.cpp
--- a/Engine/Plugins/Experimental/SlateIM/Source/SlateIM/Private/Implementation/SlateIM_Containers.cpp
+++ b/Engine/Plugins/Experimental/SlateIM/Source/SlateIM/Private/Implementation/SlateIM_Containers.cpp
@@ -278,6 +278,15 @@ namespace SlateIM
 				{
 					Scope.UpdateWidget(TableRow);
 				}
+
+				if (ParentRow)
+				{
+					TableRow->SetParent(ParentRow);
+				}
+				else
+				{
+					TableRow->SetParent(Body);
+				}				
 			}
 
 			FSlateIMManager::Get().PushContainer(FContainerNode(TableRow));
@@ -625,11 +634,8 @@ namespace SlateIM
 		return Row && Row->AreTableRowContentsRequired();
 	}
 
-	bool BeginTableRowChildren(uint32 ParentRowId)
+	bool BeginTableRowChildren(uint32 ParentRowId, bool bDefaultExpanded)
 	{
-		// TODO - Let caller set default expansion state
-		// TODO - Save and restore expansion state
-		
 		// If there's an open cell, close it
 		if (TSharedPtr<SImStackBox> Cell = FSlateIMManager::Get().GetCurrentContainer<SImStackBox>())
 		{
@@ -650,9 +656,35 @@ namespace SlateIM
 		{
 			if (TSharedPtr<SImTableView> OwningTable = ParentRow->GetOwningTable())
 			{
-				if (TSharedPtr<FSlateIMTableBody> TableBody = OwningTable->GetBody())
+				bool bHasExpansionState = false;
+				bool bExpanded = false;
+				TSharedPtr<ISlateIMContainer> GrandParent = ParentRow->GetParent();
+
+				if (GrandParent->IsA<FSlateIMTableBody>())
+				{
+					TSharedPtr<FSlateIMTableBody> GrandParentRow = StaticCastSharedPtr<FSlateIMTableBody>(GrandParent);
+					
+					if (GrandParentRow->HasSavedExpansionState(ParentRowId))
+					{
+						bHasExpansionState = true;
+						bExpanded = GrandParentRow->GetSavedExpansionState(ParentRowId, bDefaultExpanded);
+					}
+				}
+				else
+				{
+					TSharedPtr<FSlateIMTableRow> GrandParentBody = StaticCastSharedPtr<FSlateIMTableRow>(GrandParent);
+
+					if (GrandParentBody->HasSavedExpansionState(ParentRowId))
+					{
+						bHasExpansionState = true;
+						bExpanded = GrandParentBody->GetSavedExpansionState(ParentRowId, bDefaultExpanded);
+					}
+				}
+
+				if (bHasExpansionState)
 				{
-					OwningTable->SetItemExpansion(ParentRow.ToSharedRef(), TableBody->GetSavedExpansionState(ParentRowId));
+					OwningTable->SetItemExpansion(ParentRow.ToSharedRef(), bExpanded);
+					return bExpanded;
 				}
 			}
 		}
diff --git a/Engine/Plugins/Experimental/SlateIM/Source/SlateIM/Private/Misc/SlateIMExamples.cpp b/Engine/Plugins/Experimental/SlateIM/Source/SlateIM/Private/Misc/SlateIMExamples.cpp
--- a/Engine/Plugins/Experimental/SlateIM/Source/SlateIM/Private/Misc/SlateIMExamples.cpp
+++ b/Engine/Plugins/Experimental/SlateIM/Source/SlateIM/Private/Misc/SlateIMExamples.cpp
@@ -688,6 +688,8 @@ void FSlateIMTestWidget::DrawTables()
 				{
 					if (bShouldAddCanada)
 					{
+						int32 ChildRowId = 1;
+
 						if (SlateIM::NextTableCell())
 						{
 							SlateIM::VAlign(VAlign_Center);
@@ -711,7 +713,7 @@ void FSlateIMTestWidget::DrawTables()
 								SlateIM::VAlign(VAlign_Center);
 								SlateIM::Text(TEXT("Province"));
 							}
-							if (SlateIM::BeginTableRowChildren())
+							if (SlateIM::BeginTableRowChildren(ChildRowId++))
 							{
 								if (SlateIM::NextTableCell())
 								{
@@ -736,7 +738,7 @@ void FSlateIMTestWidget::DrawTables()
 								SlateIM::VAlign(VAlign_Center);
 								SlateIM::Text(TEXT("Province"));
 							}
-							if (SlateIM::BeginTableRowChildren())
+							if (SlateIM::BeginTableRowChildren(ChildRowId++))
 							{
 								if (SlateIM::NextTableCell())
 								{
diff --git a/Engine/Plugins/Experimental/SlateIM/Source/SlateIM/Public/SlateIM.h b/Engine/Plugins/Experimental/SlateIM/Source/SlateIM/Public/SlateIM.h
--- a/Engine/Plugins/Experimental/SlateIM/Source/SlateIM/Public/SlateIM.h
+++ b/Engine/Plugins/Experimental/SlateIM/Source/SlateIM/Public/SlateIM.h
@@ -283,14 +283,17 @@ namespace SlateIM
 
 	/**
 	 * Begin adding content as a child to the current table row. NextTableCell() must still be called after this function.
-	 * In the case of a dynamic table, providing a parent row id ensures that its expansion state is maintained correctly.
-	 * A parent row id is not required if the table is not dynamic.
+	 * If the table is a tree, a unique row id must be provided when starting a new set of child rows so that the expansion
+	 * state of the parent row is saved. Not providing a row id will cause the parent row to revert to the default expansion
+	 * state if its grandparent row is closed.
 	 * 
-	 * @param ParentRowId Id of the parent row.
+	 * @param ParentRowId Id of the parent row. 0 to disable state saving. Row ids must be unique to its parent.
+	 * @param bDefaultExpanded (optional) Set the default expansion state of the parent (whether the children following are visible.)
+	 *   Only applies to rows with a parent id.
 	 * 
 	 * @return Whether the parent row is expanded or not. When false, the child content can be skipped to save cycles.
 	 */
-	SLATEIM_API bool BeginTableRowChildren(uint32 ParentRowId = 0);
+	SLATEIM_API bool BeginTableRowChildren(uint32 ParentRowId = 0, bool bDefaultExpanded = false);
 	/**
 	 * Stop adding child content to the parent row. This must be called no matter the result of BeginTableRowChildren()
 	 */
