SlateIM: You can now create always on top windows.

#Synced-CL 50113862

diff --git a/Engine/Plugins/Experimental/SlateIM/Source/SlateIM/Private/Implementation/SlateIM_Roots.cpp b/Engine/Plugins/Experimental/SlateIM/Source/SlateIM/Private/Implementation/SlateIM_Roots.cpp
--- a/Engine/Plugins/Experimental/SlateIM/Source/SlateIM/Private/Implementation/SlateIM_Roots.cpp
+++ b/Engine/Plugins/Experimental/SlateIM/Source/SlateIM/Private/Implementation/SlateIM_Roots.cpp
@@ -1,7 +1,8 @@
-﻿// Copyright Epic Games, Inc. All Rights Reserved.
+// Copyright Epic Games, Inc. All Rights Reserved.
 
 #include "SlateIM.h"
 
+#include "Containers/Set.h"
 #include "Framework/Application/SlateApplication.h"
 #include "Misc/SlateIMManager.h"
 #include "Roots/SlateIMExposedRoot.h"
@@ -16,6 +17,18 @@
 namespace SlateIM
 {
 	bool BeginWindowRoot(FName UniqueName, const FStringView& WindowTitle, FVector2f WindowSize, bool bShouldReopen)
+	{
+		return BeginWindowRoot(
+			UniqueName,
+			{
+				.WindowTitle   = WindowTitle,
+				.WindowSize    = WindowSize,
+				.bShouldReopen = bShouldReopen
+			}
+		);
+	}
+	
+	bool BeginWindowRoot(FName UniqueName, const FWindowParams& Params)
 	{
 		if (!FSlateIMManager::Get().CanUpdateSlateIM())
 		{
@@ -28,7 +41,7 @@ namespace SlateIM
 		if (RootNode)
 		{
 			// Window existed and was destroyed and needs to be reopened if requested
-			bNeedsCreate = !RootNode->RootWidget.IsValid() && bShouldReopen;
+			bNeedsCreate = !RootNode->RootWidget.IsValid() && Params.bShouldReopen;
 
 			RootNode->RootState = RootNode->RootWidget->IsVisible();
 		}
@@ -39,8 +52,9 @@ namespace SlateIM
 			// TODO - Save and restore window size and position, mapped to UniqueName
 			TSharedRef<SWindow> NewWindowWidget
 				= SNew(SWindow)
-				.Title(FText::FromStringView(WindowTitle))
-				.ClientSize(WindowSize);
+				.Title(FText::FromStringView(Params.WindowTitle))
+				.ClientSize(Params.WindowSize)
+				.IsTopmostWindow(Params.bAlwaysOnTop);
 
 			FSlateApplication::Get().AddWindow(NewWindowWidget);
 
@@ -61,7 +75,7 @@ namespace SlateIM
 			WindowRoot = StaticCastSharedPtr<FSlateIMWindowRoot>(RootNode->RootWidget);
 		}
 		
-		WindowRoot->UpdateWindow(WindowTitle);
+		WindowRoot->UpdateWindow(Params.WindowTitle);
 
 		FSlateIMManager::Get().BeginRoot(UniqueName);
 
diff --git a/Engine/Plugins/Experimental/SlateIM/Source/SlateIM/Private/Misc/SlateIMExamples.h b/Engine/Plugins/Experimental/SlateIM/Source/SlateIM/Private/Misc/SlateIMExamples.h
--- a/Engine/Plugins/Experimental/SlateIM/Source/SlateIM/Private/Misc/SlateIMExamples.h
+++ b/Engine/Plugins/Experimental/SlateIM/Source/SlateIM/Private/Misc/SlateIMExamples.h
@@ -70,7 +70,7 @@ class FSlateStyleBrowser : public FSlateIMWindowBase
 {
 public:
 	FSlateStyleBrowser()
-		: FSlateIMWindowBase(TEXT("SlateIM Style Browser"), FVector2f(1000, 500), TEXT("SlateIM.ToggleSlateStyleBrowser"), TEXT("Opens a window that previews available slate styles"))
+		: FSlateIMWindowBase(TEXT("SlateIM Style Browser"), FVector2f(1000, 500), /* Always On Top */ false, TEXT("SlateIM.ToggleSlateStyleBrowser"), TEXT("Opens a window that previews available slate styles"))
 	{}
 	
 	virtual void DrawWindow(float DeltaTime) override;
@@ -88,7 +88,7 @@ class FSlateIMTestWindowWidget : public FSlateIMWindowBase
 {
 public:
 	FSlateIMTestWindowWidget(const TCHAR* Command, const TCHAR* CommandHelp)
-		: FSlateIMWindowBase(TEXT("SlateIM Test Suite"), FVector2f(700, 900), Command, CommandHelp)
+		: FSlateIMWindowBase(TEXT("SlateIM Test Suite"), FVector2f(700, 900), /* Always On Top */ false, Command, CommandHelp)
 	{}
 	
 	virtual void DrawWindow(float DeltaTime) override;
diff --git a/Engine/Plugins/Experimental/SlateIM/Source/SlateIM/Private/SlateIMWidgetBase.cpp b/Engine/Plugins/Experimental/SlateIM/Source/SlateIM/Private/SlateIMWidgetBase.cpp
--- a/Engine/Plugins/Experimental/SlateIM/Source/SlateIM/Private/SlateIMWidgetBase.cpp
+++ b/Engine/Plugins/Experimental/SlateIM/Source/SlateIM/Private/SlateIMWidgetBase.cpp
@@ -57,7 +57,15 @@ void FSlateIMWindowBase::DrawWidget(float DeltaTime)
 		return;
 	}
 	
-	const bool bIsDrawingWindow = SlateIM::BeginWindowRoot(GetWidgetName(), WindowTitle, WindowSize);
+	const bool bIsDrawingWindow = SlateIM::BeginWindowRoot(
+		GetWidgetName(),
+		{
+			.WindowTitle  = WindowTitle,
+			.WindowSize   = WindowSize,
+			.bAlwaysOnTop = bAlwaysOnTop
+		}
+	);
+
 	if (bIsDrawingWindow)
 	{
 		DrawWindow(DeltaTime);
diff --git a/Engine/Plugins/Experimental/SlateIM/Source/SlateIM/Public/SlateIM.h b/Engine/Plugins/Experimental/SlateIM/Source/SlateIM/Public/SlateIM.h
--- a/Engine/Plugins/Experimental/SlateIM/Source/SlateIM/Public/SlateIM.h
+++ b/Engine/Plugins/Experimental/SlateIM/Source/SlateIM/Public/SlateIM.h
@@ -1,7 +1,6 @@
 // Copyright Epic Games, Inc. All Rights Reserved.
 
 #pragma once
-
 #include "Framework/Views/ITypedTableView.h"
 #include "Layout/Margin.h"
 #include "Modules/ModuleInterface.h"
@@ -87,6 +86,8 @@ namespace SlateIM
 	/**
 	 * Begins a new floating window root
 	 *
+	 * Please use the method below with the struct parameters.
+	 *
 	 * @param UniqueName a unique name to identify this window
 	 * @param WindowTitle The name of the window. This will also be displayed in the windows title bar
 	 * @param WindowSize The size to open the window at (does not update the size of an existing window)
@@ -98,6 +99,33 @@ namespace SlateIM
 	 */
 	SLATEIM_API bool BeginWindowRoot(FName UniqueName, const FStringView& WindowTitle, FVector2f WindowSize, bool bShouldReopen = false);
 
+	struct FWindowParams
+	{
+		// The name of the window. This will also be displayed in the windows title bar
+		FStringView WindowTitle;
+
+		// The size to open the window at (does not update the size of an existing window)
+		FVector2f WindowSize = FVector2f::ZeroVector;
+
+		// If this window was created with a previous call to BeginWindow and then closed, passing in true will reopen the window and false will leave it closed. If the window has never been seen this parameter does nothing and the window will open
+		bool bShouldReopen = false;
+
+		// Window will be create always on top.
+		bool bAlwaysOnTop = false;
+	};
+
+	/**
+	 * Begins a new floating window root
+	 *
+	 * @param UniqueName a unique name to identify this window
+	 * @param Params Configuration parameters for the window.
+	 *
+	 * @return the current window state. True if open, false if closed or not updating
+	 *
+	 * @note Consider using FSlateIMWindowBase to handle creating and updating your Slate IM Window
+	 */
+	SLATEIM_API bool BeginWindowRoot(FName UniqueName, const FWindowParams& Params = {});
+
 #if WITH_ENGINE
 	// TODO - Game viewport widgets might need something to control HitTestability and cursor visibility
 	/**
@@ -200,19 +228,21 @@ namespace SlateIM
 
 	/**
 	 * Begin a Table container. The contents are automatically scrollable if the table is sized smaller than all its content.
-	 * 
-	 * @param Params Settings for table construction. Not optional for compatibility with the below method.
+	 *
+	 * Please use the below method instead!
+	 *
+	 * @param InStyle A Table View Style to override the default style. Changes to the style after the first invocation are not reflected,
+	 *				  provide a new style object to update the visuals at runtime
+	 * @param InRowStyle (optional) A Table Row Style to override the default table row style. Changing the style only affects new table rows.
 	 */
-	SLATEIM_API void BeginTable(const FTableParams& Params);
+	SLATEIM_API void BeginTable(const FTableViewStyle* InStyle, const FTableRowStyle* InRowStyle = nullptr);
 
 	/**
 	 * Begin a Table container. The contents are automatically scrollable if the table is sized smaller than all its content.
-	 * 
-	 * @param InStyle (optional) A Table View Style to override the default style. Changes to the style after the first invocation are not reflected,
-	 *				  provide a new style object to update the visuals at runtime
-	 * @param InRowStyle (optional) A Table Row Style to override the default table row style. Changing the style only affects new table rows.
+	 *
+	 * @param Params (optional) Settings for table construction.
 	 */
-	SLATEIM_API void BeginTable(const FTableViewStyle* InStyle = nullptr, const FTableRowStyle* InRowStyle = nullptr);
+	SLATEIM_API void BeginTable(const FTableParams& Params = FTableParams());
 
 	/**
 	 * End the current Table container.
diff --git a/Engine/Plugins/Experimental/SlateIM/Source/SlateIM/Public/SlateIMWidgetBase.h b/Engine/Plugins/Experimental/SlateIM/Source/SlateIM/Public/SlateIMWidgetBase.h
--- a/Engine/Plugins/Experimental/SlateIM/Source/SlateIM/Public/SlateIMWidgetBase.h
+++ b/Engine/Plugins/Experimental/SlateIM/Source/SlateIM/Public/SlateIMWidgetBase.h
@@ -62,9 +62,14 @@ class FSlateIMWindowBase : public FSlateIMWidgetWithCommandBase
 {
 public:
 	FSlateIMWindowBase(FStringView WindowTitle, FVector2f WindowSize, const TCHAR* Command, const TCHAR* CommandHelp)
+		: FSlateIMWindowBase(WindowTitle, WindowSize, /* Always On Top */ false, Command, CommandHelp)
+	{}
+
+	FSlateIMWindowBase(FStringView WindowTitle, FVector2f WindowSize, bool bAlwaysOnTop, const TCHAR* Command, const TCHAR* CommandHelp)
 		: FSlateIMWidgetWithCommandBase(Command, CommandHelp)
 		, WindowTitle(WindowTitle)
 		, WindowSize(WindowSize)
+		, bAlwaysOnTop(bAlwaysOnTop)
 	{}
 
 protected:
@@ -75,6 +80,7 @@ private:
 
 	FString WindowTitle;
 	FVector2f WindowSize;
+	bool bAlwaysOnTop;
 };
 
 /**
diff --git a/Engine/Plugins/Experimental/SlateIM/Source/SlateIMInGame/Private/SlateIMInGameWindow.cpp b/Engine/Plugins/Experimental/SlateIM/Source/SlateIMInGame/Private/SlateIMInGameWindow.cpp
--- a/Engine/Plugins/Experimental/SlateIM/Source/SlateIMInGame/Private/SlateIMInGameWindow.cpp
+++ b/Engine/Plugins/Experimental/SlateIM/Source/SlateIMInGame/Private/SlateIMInGameWindow.cpp
@@ -18,7 +18,14 @@ ASlateIMInGameWindow::ASlateIMInGameWindow(const FName& InWindowName, const FStr
 
 void ASlateIMInGameWindow::DrawWidget(const float DeltaTime)
 {
-	const bool bIsDrawingWindow = SlateIM::BeginWindowRoot(WindowName, WindowTitle, WindowSize);
+	const bool bIsDrawingWindow = SlateIM::BeginWindowRoot(
+		WindowName,
+		{
+			.WindowTitle = WindowTitle,
+			.WindowSize = WindowSize
+		}
+	);
+
 	if (bIsDrawingWindow)
 	{
 		DrawContent(DeltaTime);
