Slate IM: Menu entries directly after submenus are now displayed correctly.

#Synced-CL 50113862

diff --git a/Engine/Plugins/Experimental/SlateIM/Source/SlateIM/Private/Containers/SImContextMenuAnchor.cpp b/Engine/Plugins/Experimental/SlateIM/Source/SlateIM/Private/Containers/SImContextMenuAnchor.cpp
--- a/Engine/Plugins/Experimental/SlateIM/Source/SlateIM/Private/Containers/SImContextMenuAnchor.cpp
+++ b/Engine/Plugins/Experimental/SlateIM/Source/SlateIM/Private/Containers/SImContextMenuAnchor.cpp
@@ -357,11 +357,14 @@ TSharedRef<SWidget> SImContextMenuAnchor::BuildMenu_Recursive(int32& InOutMenuIn
 			TSharedRef<SWidget> SubMenuWidget = BuildMenu_Recursive(InOutMenuIndex, MenuLevel+1);
 			int32 SubMenuIndex = SubMenuWidgets.Add(SubMenuWidget);
 			MenuBuilder.AddWrapperSubMenu(FText::FromStringView(MenuText), FText::FromStringView(ToolTipText), FOnGetContent::CreateSP(this, &SImContextMenuAnchor::OnGetSubMenuContent, SubMenuIndex), FSlateIcon());
-	
 			break;
 		}
 
-		++InOutMenuIndex;
+		// This has already been done in the switch statement for SubMenus.
+		if (MenuItemData.Type != EMenuType::SubMenu)
+		{
+			++InOutMenuIndex;
+		}
 	}
 
 	if (bHasCurrentSection)
