SlateIM: Partial fix to content in tables jittering about.

#Synced-CL 50214844

diff --git a/Engine/Plugins/Experimental/SlateIM/Source/SlateIM/Private/Containers/SImTableView.cpp b/Engine/Plugins/Experimental/SlateIM/Source/SlateIM/Private/Containers/SImTableView.cpp
--- a/Engine/Plugins/Experimental/SlateIM/Source/SlateIM/Private/Containers/SImTableView.cpp
+++ b/Engine/Plugins/Experimental/SlateIM/Source/SlateIM/Private/Containers/SImTableView.cpp
@@ -143,12 +143,12 @@ void FSlateIMTableRow::UpdateChild(FSlateIMChild Child, int32 Index, const FSlat
 	else
 	{
 		Children.Add(Child);
-	}
 
-	// Modifying the child count requires a rebuild to handle the case where we need to create/destroy the expander arrow and child widgets
-	if (OwningTable)
-	{
-		OwningTable->RequestRebuild();
+		// Modifying the child count requires a rebuild to handle the case where we need to create/destroy the expander arrow and child widgets
+		if (OwningTable)
+		{
+			OwningTable->RequestRebuild();
+		}
 	}
 }
 
@@ -395,12 +395,12 @@ void FSlateIMTableBody::UpdateChild(FSlateIMChild Child, int32 Index, const FSla
 	else
 	{
 		Children.Add(ChildRow.ToSharedRef());
-	}
 
-	// Modifying the child count requires a rebuild to handle the case where we need to create/destroy the expander arrow and child widgets
-	if (OwningTable)
-	{
-		OwningTable->RequestRebuild();
+		// Modifying the child count requires a rebuild to handle the case where we need to create/destroy the expander arrow and child widgets
+		if (OwningTable)
+		{
+			OwningTable->RequestRebuild();
+		}
 	}
 }
 
