SlateIM: Within a root, you can now specify a tab id for 1-time activation, if it exists.

#Synced-CL 50388552

diff --git a/Engine/Plugins/Experimental/SlateIM/Source/SlateIM/Private/Implementation/SlateIM_Tabs.cpp b/Engine/Plugins/Experimental/SlateIM/Source/SlateIM/Private/Implementation/SlateIM_Tabs.cpp
--- a/Engine/Plugins/Experimental/SlateIM/Source/SlateIM/Private/Implementation/SlateIM_Tabs.cpp
+++ b/Engine/Plugins/Experimental/SlateIM/Source/SlateIM/Private/Implementation/SlateIM_Tabs.cpp
@@ -114,6 +114,11 @@ namespace SlateIM
 				ContainerWidget->TabIcon = TabIcon;
 				ContainerWidget->TabTitle = TabTitle;
 			}
+
+			if (!FSlateIMManager::Get().GetMutableCurrentRoot().Tabs.Contains(TabId))
+			{
+				FSlateIMManager::Get().GetMutableCurrentRoot().Tabs.Add(TabId, ContainerWidget);
+			}
 		}
 
 		FSlateIMManager::Get().PushContainer(FContainerNode(ContainerWidget));
@@ -125,4 +130,12 @@ namespace SlateIM
 	{
 		FSlateIMManager::Get().PopContainer<SImTab>();
 	}
+	
+	void ActivateTab(const FName& TabId)
+	{
+		if (ensureMsgf(FSlateIMManager::Get().HasCurrentRoot(), TEXT("ActivateTab called without a current root. Have you started a root?!")))
+		{
+			FSlateIMManager::Get().GetMutableCurrentRoot().TabsToActivate.Add(TabId);
+		}
+	}
 }
diff --git a/Engine/Plugins/Experimental/SlateIM/Source/SlateIM/Private/Misc/SlateIMExamples.cpp b/Engine/Plugins/Experimental/SlateIM/Source/SlateIM/Private/Misc/SlateIMExamples.cpp
--- a/Engine/Plugins/Experimental/SlateIM/Source/SlateIM/Private/Misc/SlateIMExamples.cpp
+++ b/Engine/Plugins/Experimental/SlateIM/Source/SlateIM/Private/Misc/SlateIMExamples.cpp
@@ -1305,6 +1305,7 @@ void FSlateIMTestWidget::DrawTabs()
 		SlateIM::BeginHorizontalStack();
 		if (SlateIM::Button(TEXT("Add Tab")))
 		{
+			SlateIM::ActivateTab(*FString::Printf(TEXT("DynamicTab%d"), DynamicTabCount));
 			++DynamicTabCount;
 		}
 		if (SlateIM::Button(TEXT("Remove Tab")))
diff --git a/Engine/Plugins/Experimental/SlateIM/Source/SlateIM/Private/Misc/SlateIMManager.cpp b/Engine/Plugins/Experimental/SlateIM/Source/SlateIM/Private/Misc/SlateIMManager.cpp
--- a/Engine/Plugins/Experimental/SlateIM/Source/SlateIM/Private/Misc/SlateIMManager.cpp
+++ b/Engine/Plugins/Experimental/SlateIM/Source/SlateIM/Private/Misc/SlateIMManager.cpp
@@ -4,6 +4,7 @@
 
 #include "Containers/SImContextMenuAnchor.h"
 #include "Containers/SImStackBox.h"
+#include "Containers/SImTab.h"
 #include "Framework/Application/SlateApplication.h"
 #include "Misc/App.h"
 #include "Roots/ISlateIMRoot.h"
@@ -65,6 +66,27 @@ namespace SlateIM
 		}
 	}
 
+	void FRootNode::ActivateTab()
+	{
+		if (!TabsToActivate.IsEmpty())
+		{
+			for (const FName& TabId : TabsToActivate)
+			{
+				if (const TWeakPtr<SImTab>* WeakTabPtr = Tabs.Find(TabId))
+				{
+					if (TSharedPtr<SImTab> Tab = WeakTabPtr->Pin())
+					{
+						Tab->ActivateInParent(ETabActivationCause::SetDirectly);
+					}
+				}
+			}
+
+			TabsToActivate.Empty();
+		}
+
+		Tabs.Empty();
+	}
+
 	TUniquePtr<FSlateIMManager> FSlateIMManager::Instance;
 
 	void FSlateIMManager::Initialize()
@@ -168,6 +190,8 @@ namespace SlateIM
 			CurrentRoot->WidgetTreeDataHash.RemoveAt(FirstUnusedIndex, CurrentRoot->WidgetTreeDataHash.Num() - FirstUnusedIndex, EAllowShrinking::No);
 		}
 
+		CurrentRoot->ActivateTab();
+
 		CurrentRoot->CurrentWidgetIndex = 0;
 		CurrentRoot = nullptr;
 
@@ -277,6 +301,18 @@ namespace SlateIM
 		return ChildWidget;
 	}
 
+	const FRootNode& FSlateIMManager::GetCurrentRoot() const
+	{
+		check(CurrentRoot != nullptr);
+		return *CurrentRoot; 
+	}
+
+	FRootNode& FSlateIMManager::GetMutableCurrentRoot()
+	{ 
+		check(CurrentRoot != nullptr);
+		return *CurrentRoot;
+	}
+
 	void FSlateIMManager::UpdateCurrentChild(FSlateIMChild Child, const FSlateIMSlotData& AlignmentData)
 	{
 		FContainerNode& CurrentContainer = CurrentRoot->CurrentContainerStack.Top();
diff --git a/Engine/Plugins/Experimental/SlateIM/Source/SlateIM/Private/Misc/SlateIMManager.h b/Engine/Plugins/Experimental/SlateIM/Source/SlateIM/Private/Misc/SlateIMManager.h
--- a/Engine/Plugins/Experimental/SlateIM/Source/SlateIM/Private/Misc/SlateIMManager.h
+++ b/Engine/Plugins/Experimental/SlateIM/Source/SlateIM/Private/Misc/SlateIMManager.h
@@ -23,6 +23,7 @@ class ISlateIMChild;
 class ISlateIMRoot;
 class ISlateIMContainer;
 class SImContextMenuAnchor;
+class SImTab;
 class SWidget;
 
 struct FSlateIMSlotData;
@@ -76,6 +77,8 @@ namespace SlateIM
 
 		void SetNextToolTip(const FStringView& InNextToolTip);
 
+		void ActivateTab();
+
 		FContainerNode RootContainer;
 
 		TArray<FContainerNode> CurrentContainerStack;
@@ -85,6 +88,10 @@ namespace SlateIM
 
 		TSharedPtr<ISlateIMRoot> RootWidget;
 
+		// Store a list of tabs so that when ActivateTab is called at the end of layout, any registered tab can be activated.
+		TMap<FName, TWeakPtr<SImTab>> Tabs;
+		TSet<FName> TabsToActivate;
+
 		FString CurrentToolTip;
 		uint64 DataHash = 0;
 		FName RootName = NAME_None;
@@ -158,8 +165,9 @@ namespace SlateIM
 		/** Gets the widget which should be at the current spot in the hierarchy or null if it doesn't exist */
 		FSlateIMChild GetCurrentChild() const;
 		TSharedPtr<SWidget> GetCurrentChildAsWidget() const;
-		const FRootNode& GetCurrentRoot() const { return *CurrentRoot; } 
-		FRootNode& GetMutableCurrentRoot() { return *CurrentRoot; }
+		const FRootNode& GetCurrentRoot() const;
+		FRootNode& GetMutableCurrentRoot();
+		bool HasCurrentRoot() const { return !!CurrentRoot; }
 
 		void UpdateCurrentChild(FSlateIMChild Child, const FSlateIMSlotData& AlignmentData);
 		void AdvanceToNextWidget();
diff --git a/Engine/Plugins/Experimental/SlateIM/Source/SlateIM/Public/SlateIM.h b/Engine/Plugins/Experimental/SlateIM/Source/SlateIM/Public/SlateIM.h
--- a/Engine/Plugins/Experimental/SlateIM/Source/SlateIM/Public/SlateIM.h
+++ b/Engine/Plugins/Experimental/SlateIM/Source/SlateIM/Public/SlateIM.h
@@ -784,6 +784,17 @@ namespace SlateIM
 	 * Ends the current Tab. This must be called even when BeginTab returns false.
 	 */
 	SLATEIM_API void EndTab();
+
+	/**
+	 * Causes a one-time activation of a tab in its parent tab well. Can be used at any point during layout.
+	 * 
+	 * To use this effectively, all tab ids should be unique. Only the first tab registered with the given id will be activated.
+	 * Can be called multiple times to activate multiple tabs in the same update.
+	 * @see ActivateInParent(ETabActivationCause).
+	 * 
+	 * @param TabId The tab id to activate.
+	 */
+	SLATEIM_API void ActivateTab(const FName& TabId);
 #pragma endregion Tabs
 
 #pragma region Util
