SlateIM: Massively increased performance on very large tables.

#Synced-CL 50388552

diff --git a/Engine/Plugins/Experimental/SlateIM/Source/SlateIM/Private/Containers/SImTableView.cpp b/Engine/Plugins/Experimental/SlateIM/Source/SlateIM/Private/Containers/SImTableView.cpp
--- a/Engine/Plugins/Experimental/SlateIM/Source/SlateIM/Private/Containers/SImTableView.cpp
+++ b/Engine/Plugins/Experimental/SlateIM/Source/SlateIM/Private/Containers/SImTableView.cpp
@@ -299,7 +299,7 @@ bool FSlateIMTableRow::AreTableRowContentsRequired()
 		// This is probably a bigger window than necessary but should handle large scroll movements and should be small enough to not impact perf
 		const float MinNearlyLiveWidget = FMath::Max(0.f, CurrentScrollOffset - NumLiveWidgets);
 		const float MaxNearlyLiveWidget = CurrentScrollOffset + (2.f * NumLiveWidgets);
-		const float ThisRowIndex = OwningTable->GetRowLinearizedIndex(AsShared());
+		const float ThisRowIndex = OwningTable->GetRowCount() - 1;
 		if (ThisRowIndex >= MinNearlyLiveWidget && ThisRowIndex <= MaxNearlyLiveWidget)
 		{
 			return true;
@@ -582,6 +582,7 @@ void SImTableView::EndColumnUpdates()
 void SImTableView::BeginTableUpdates()
 {
 	ColumnCount = 0;
+	RowCount = 0;
 }
 
 void SImTableView::EndTableUpdates()
@@ -680,6 +681,16 @@ int32 SImTableView::FindColumnIndex(const FName& ColumnName) const
 	return HeaderRow ? HeaderRow->FindColumnIndex(ColumnName) : INDEX_NONE;
 }
 
+void SImTableView::OnRowAdded()
+{
+	++RowCount;
+}
+
+int32 SImTableView::GetRowCount() const
+{
+	return RowCount;
+}
+
 void SImTableView::AddColumn_Internal(const FName& ColumnId, const FStringView& ColumnToolTip, const FSlateIMSlotData& SlotData)
 {
 	SHeaderRow::FColumn::FArguments Args = SHeaderRow::FColumn::FArguments()
diff --git a/Engine/Plugins/Experimental/SlateIM/Source/SlateIM/Private/Containers/SImTableView.h b/Engine/Plugins/Experimental/SlateIM/Source/SlateIM/Private/Containers/SImTableView.h
--- a/Engine/Plugins/Experimental/SlateIM/Source/SlateIM/Private/Containers/SImTableView.h
+++ b/Engine/Plugins/Experimental/SlateIM/Source/SlateIM/Private/Containers/SImTableView.h
@@ -218,6 +218,9 @@ public:
 	bool IsSameColumnCount(const TSharedRef<FSlateIMTableRow>& Row) const;
 	int32 FindColumnIndex(const FName& ColumnName) const;
 
+	void OnRowAdded();
+	int32 GetRowCount() const;
+
 private:
 	void AddColumn_Internal(const FName& ColumnId, const FStringView& ColumnToolTip, const FSlateIMSlotData& SlotData);
 
@@ -253,6 +256,7 @@ private:
 	TSharedPtr<FSlateIMTableBody> Body;
 	const FTableRowStyle* RowStyle = nullptr;
 	int32 ColumnCount = 0;
+	int32 RowCount = 0;
 
 	mutable int32 CachedNumLiveWidgets = 0;
 
diff --git a/Engine/Plugins/Experimental/SlateIM/Source/SlateIM/Private/Implementation/SlateIM_Containers.cpp b/Engine/Plugins/Experimental/SlateIM/Source/SlateIM/Private/Implementation/SlateIM_Containers.cpp
--- a/Engine/Plugins/Experimental/SlateIM/Source/SlateIM/Private/Implementation/SlateIM_Containers.cpp
+++ b/Engine/Plugins/Experimental/SlateIM/Source/SlateIM/Private/Implementation/SlateIM_Containers.cpp
@@ -286,7 +286,12 @@ namespace SlateIM
 				else
 				{
 					TableRow->SetParent(Body);
-				}				
+				}
+
+				if (Table)
+				{
+					Table->OnRowAdded();
+				}
 			}
 
 			FSlateIMManager::Get().PushContainer(FContainerNode(TableRow));
@@ -628,10 +633,17 @@ namespace SlateIM
 			}
 			// Create a vertical stack to act as our cell widget
 			// TODO - Only automatically create a container when one is not created by the user
-			SlateIM::BeginVerticalStack();
+
+			if (Row->AreTableRowContentsRequired())
+			{
+				SlateIM::BeginVerticalStack();
+				return true;
+			}
+
+			SlateIM::Widget(SNullWidget::NullWidget);
 		}
 
-		return Row && Row->AreTableRowContentsRequired();
+		return false;
 	}
 
 	bool BeginTableRowChildren(uint32 ParentRowId, bool bDefaultExpanded)
