SlateIM: Backing out part of a change which adversely affected table rendering.

#Synced-CL 50388552

diff --git a/Engine/Plugins/Experimental/SlateIM/Source/SlateIM/Private/Implementation/SlateIM_Containers.cpp b/Engine/Plugins/Experimental/SlateIM/Source/SlateIM/Private/Implementation/SlateIM_Containers.cpp
--- a/Engine/Plugins/Experimental/SlateIM/Source/SlateIM/Private/Implementation/SlateIM_Containers.cpp
+++ b/Engine/Plugins/Experimental/SlateIM/Source/SlateIM/Private/Implementation/SlateIM_Containers.cpp
@@ -631,16 +631,13 @@ namespace SlateIM
 			{
 				FSlateIMManager::Get().NextHAlign = HAlign_Fill;
 			}
+			
 			// Create a vertical stack to act as our cell widget
 			// TODO - Only automatically create a container when one is not created by the user
 
-			if (Row->AreTableRowContentsRequired())
-			{
-				SlateIM::BeginVerticalStack();
-				return true;
-			}
+			SlateIM::BeginVerticalStack();
 
-			SlateIM::Widget(SNullWidget::NullWidget);
+			return Row->AreTableRowContentsRequired();
 		}
 
 		return false;
