SlateIM: Tabbing into undocked Tabs in TabGroups inside Window Roots, which are not in other Tabs, will now no longer cause the main Window Root to hide behind the editor.

#Synced-CL 50687911

diff --git a/Engine/Plugins/Experimental/SlateIM/Source/SlateIM/Private/Containers/SImTabGroup.cpp b/Engine/Plugins/Experimental/SlateIM/Source/SlateIM/Private/Containers/SImTabGroup.cpp
--- a/Engine/Plugins/Experimental/SlateIM/Source/SlateIM/Private/Containers/SImTabGroup.cpp
+++ b/Engine/Plugins/Experimental/SlateIM/Source/SlateIM/Private/Containers/SImTabGroup.cpp
@@ -3,6 +3,7 @@
 
 #include "SImTabGroup.h"
 
+#include "Framework/Docking/TabManager.h"
 #include "Misc/SlateIMManager.h"
 #include "Roots/SlateIMWindowRoot.h"
 #include "SImTab.h"
@@ -137,9 +138,25 @@ void SImTabGroup::PrivateRegisterAttributes(FSlateAttributeInitializer&)
 void SImTabGroup::Construct(const FArguments& InArgs)
 {
 	TabGroupId = InArgs._TabGroupId;
-	const TSharedPtr<SDockTab> ParentTab = SlateIM::FSlateIMManager::Get().FindMostRecentContainer<SImTab>();
-	TabManager = FGlobalTabmanager::Get()->NewTabManager(ParentTab.IsValid() ? ParentTab.ToSharedRef() : SNew(SDockTab));
-	const TSharedRef<FTabManager::FArea> PrimaryArea = FTabManager::NewPrimaryArea();
+	
+	if (TSharedPtr<SDockTab> ParentTab = SlateIM::FSlateIMManager::Get().FindMostRecentContainer<SImTab>())
+	{
+		TabManager = FGlobalTabmanager::Get()->NewTabManager(ParentTab.ToSharedRef());
+	}
+	else if (SlateIM::FSlateIMManager::Get().HasCurrentRoot()
+		&& SlateIM::FSlateIMManager::Get().GetCurrentRoot().RootWidget.IsValid()
+		&& SlateIM::FSlateIMManager::Get().GetCurrentRoot().RootWidget->IsA<FSlateIMWindowRoot>())
+	{
+		TSharedRef<FSlateIMWindowRoot> WindowRoot = StaticCastSharedPtr<FSlateIMWindowRoot>(SlateIM::FSlateIMManager::Get().GetCurrentRoot().RootWidget).ToSharedRef();
+		TabManager = FGlobalTabmanager::Get()->NewTabManager(WindowRoot->GetRootTab());
+	}
+	else
+	{
+		ensureAlwaysMsgf(false, TEXT("Missing root widget when creating a tab. Are you missing a Begin Root?"));
+		return;
+	}
+
+	const TSharedRef<FTabManager::FArea> PrimaryArea = TabManager->NewPrimaryArea();
 	WorkingLayoutNodeStack.Emplace(FImTabLayoutNode{ PrimaryArea });
 	TabLayout = FTabManager::NewLayout(TabGroupId)->AddArea(PrimaryArea);
 }
diff --git a/Engine/Plugins/Experimental/SlateIM/Source/SlateIM/Private/Roots/SlateIMWindowRoot.cpp b/Engine/Plugins/Experimental/SlateIM/Source/SlateIM/Private/Roots/SlateIMWindowRoot.cpp
--- a/Engine/Plugins/Experimental/SlateIM/Source/SlateIM/Private/Roots/SlateIMWindowRoot.cpp
+++ b/Engine/Plugins/Experimental/SlateIM/Source/SlateIM/Private/Roots/SlateIMWindowRoot.cpp
@@ -4,10 +4,12 @@
 #include "SlateIMWindowRoot.h"
 
 #include "Framework/Application/SlateApplication.h"
+#include "Widgets/Docking/SDockTab.h"
 #include "Widgets/SImWrapper.h"
 
 FSlateIMWindowRoot::FSlateIMWindowRoot(TSharedRef<SWindow> Window)
 	: RootWindow(Window)
+	, RootTab(SNew(SDockTab))
 {}
 
 FSlateIMWindowRoot::~FSlateIMWindowRoot()
@@ -61,3 +63,8 @@ TSharedPtr<SWindow> FSlateIMWindowRoot::GetWindow() const
 {
 	return RootWindow.Pin();
 }
+
+TSharedRef<SDockTab> FSlateIMWindowRoot::GetRootTab() const
+{
+	return RootTab;
+}
diff --git a/Engine/Plugins/Experimental/SlateIM/Source/SlateIM/Private/Roots/SlateIMWindowRoot.h b/Engine/Plugins/Experimental/SlateIM/Source/SlateIM/Private/Roots/SlateIMWindowRoot.h
--- a/Engine/Plugins/Experimental/SlateIM/Source/SlateIM/Private/Roots/SlateIMWindowRoot.h
+++ b/Engine/Plugins/Experimental/SlateIM/Source/SlateIM/Private/Roots/SlateIMWindowRoot.h
@@ -4,6 +4,7 @@
 
 #include "ISlateIMRoot.h"
 
+class SDockTab;
 class SImWrapper;
 class SWindow;
 
@@ -22,9 +23,11 @@ public:
 	void UpdateWindow(const FStringView& Title);
 
 	TSharedPtr<SWindow> GetWindow() const;
+
+	TSharedRef<SDockTab> GetRootTab() const;
 	
 private:
 	TWeakPtr<SWindow> RootWindow;
 	TSharedPtr<SImWrapper> WindowRootWidget;
-
+	TSharedRef<SDockTab> RootTab;
 };
