Add SlateIM string filter utility class
#Synced-CL 50970435
#okfor

#Synced-CL 50970435

diff --git a/Engine/Plugins/Experimental/SlateIM/Source/SlateIM/Private/Misc/SlateIMExamples.cpp b/Engine/Plugins/Experimental/SlateIM/Source/SlateIM/Private/Misc/SlateIMExamples.cpp
--- a/Engine/Plugins/Experimental/SlateIM/Source/SlateIM/Private/Misc/SlateIMExamples.cpp
+++ b/Engine/Plugins/Experimental/SlateIM/Source/SlateIM/Private/Misc/SlateIMExamples.cpp
@@ -167,6 +167,12 @@ void FSlateIMTestWidget::Draw()
 		DrawTabs();
 	}
 	SlateIM::EndTab();
+
+	if (SlateIM::BeginTab(TEXT("Filter")))
+	{
+		DrawFilter();
+	}
+	SlateIM::EndTab();
 	SlateIM::EndTabStack();
 	SlateIM::EndTabGroup();
 	SlateIM::EndBorder();
@@ -1384,6 +1390,50 @@ void FSlateIMTestWidget::DrawTabs()
 	}
 }
 
+void FSlateIMTestWidget::DrawFilter()
+{
+	SCOPED_NAMED_EVENT_TEXT("Filter", FColorList::Goldenrod);
+	SlateIM::Fill();
+	SlateIM::HAlign(HAlign_Fill);
+	SlateIM::VAlign(VAlign_Fill);
+	SlateIM::Padding({ 5.f });
+	SlateIM::BeginVerticalStack();
+	{
+		static TArray<FString> FruitList = {
+			TEXT("Apple"),
+			TEXT("Banana"),
+			TEXT("Cherry"),
+			TEXT("Dragonfruit"),
+			TEXT("Elderberry"),
+			TEXT("Fig"),
+			TEXT("Grape"),
+			TEXT("Honeydew"),
+			TEXT("Kiwi"),
+			TEXT("Lemon")
+		};
+
+		if (!FilteredFruitList.IsSet())
+		{
+			FilteredFruitList = FruitList;
+		}
+
+		// Draw the filter input and filter the list
+		bool bFruitListChanged = false;
+		if (SlateIM::TextFilter(FruitFilter, TEXT("Search fruits...")))
+		{
+			FilteredFruitList = FruitFilter.Filter(FruitList);
+			bFruitListChanged = true;
+		}
+
+		// Display the filtered list
+		SlateIM::Fill();
+		SlateIM::HAlign(HAlign_Fill);
+		int32 SelectedFruitIndex = INDEX_NONE;
+		SlateIM::SelectionList(FilteredFruitList.GetValue(), SelectedFruitIndex, bFruitListChanged);
+	}
+	SlateIM::EndVerticalStack();
+}
+
 void FSlateStyleBrowser::DrawWindow(float DeltaTime)
 {
 	static const TArray<FString> Options = { TEXT("Option 1"), TEXT("Option 2"), TEXT("Option 3") };
diff --git a/Engine/Plugins/Experimental/SlateIM/Source/SlateIM/Private/Misc/SlateIMExamples.h b/Engine/Plugins/Experimental/SlateIM/Source/SlateIM/Private/Misc/SlateIMExamples.h
--- a/Engine/Plugins/Experimental/SlateIM/Source/SlateIM/Private/Misc/SlateIMExamples.h
+++ b/Engine/Plugins/Experimental/SlateIM/Source/SlateIM/Private/Misc/SlateIMExamples.h
@@ -4,6 +4,7 @@
 #include "Containers/RingBuffer.h"
 
 #if WITH_SLATEIM_EXAMPLES
+#include "Misc/SlateIMStringFilter.h"
 #include "SlateIM.h"
 #include "SlateIMWidgetBase.h"
 #include "Styling/SlateTypes.h"
@@ -19,6 +20,7 @@ public:
 	void DrawGraphs();
 	void DrawInputs();
 	void DrawTabs();
+	void DrawFilter();
 	
 private:
 	double CurrentTime = 0;
@@ -46,6 +48,8 @@ private:
 	TArray<FString> ComboBoxItems = { TEXT("Option 1"), TEXT("Option 2"), TEXT("Option 3"), TEXT("Option 4") };
 	bool bRefreshComboItems = false;
 	int32 DynamicTabCount = 3;
+	FSlateIMStringFilter FruitFilter;
+	TOptional<TArray<FString>> FilteredFruitList;
 	TRingBuffer<double> SquareGraphValues;
 	TRingBuffer<FVector2D> SinGraphValues = { FVector2D(0, FMath::Sin(0.0)) };
 	TRingBuffer<FVector2D> CosGraphValues = { FVector2D(0, FMath::Cos(0.0)) };
diff --git a/Engine/Plugins/Experimental/SlateIM/Source/SlateIM/Public/Misc/SlateIMStringFilter.h b/Engine/Plugins/Experimental/SlateIM/Source/SlateIM/Public/Misc/SlateIMStringFilter.h
new file mode 100644
--- /dev/null
+++ b/Engine/Plugins/Experimental/SlateIM/Source/SlateIM/Public/Misc/SlateIMStringFilter.h
@@ -0,0 +1,72 @@
+// Copyright Epic Games, Inc. All Rights Reserved.
+
+#pragma once
+
+#include "Containers/StringFwd.h"
+#include "Containers/UnrealString.h"
+#include "Misc/TextFilter.h"
+
+/**
+ * Base template class for SlateIM text filters.
+ *
+ * Provides a common interface for filtering items using TTextFilter.
+ */
+template<typename ItemType>
+class FSlateImTextFilter
+{
+public:
+	FSlateImTextFilter() = default;
+
+	virtual ~FSlateImTextFilter() = default;
+
+	/**
+	 * Filters an array of items, returning a new array with only matching entries.
+	 * @param InArray The source array to filter
+	 * @return A new array containing only the items that pass the filter
+	 */
+	using ArrayItemType = std::remove_const_t<typename TRemoveReference<ItemType>::Type>;
+	TArray<ArrayItemType> Filter(const TArray<ArrayItemType>& InArray) const
+	{
+		if (!TextFilter.IsValid())
+		{
+			return InArray;
+		}
+
+		return InArray.FilterByPredicate([this](const ArrayItemType& Item)
+		{
+			return TextFilter->PassesFilter(Item);
+		});
+	}
+
+	FString InputText;
+
+	TSharedPtr<TTextFilter<ItemType>> TextFilter;
+};
+
+/**
+ * SlateIM filter utility class for filtering FStrings in immediate-mode menus.
+ * Usage:
+ *   // As a member variable in your SlateIM menu class
+ *   FSlateIMStringFilter MyFilter;
+ *
+ *   // In your immediate-mode drawing code
+ *   SlateIM::TextFilter(MyFilter, TEXT("Filter"));
+ *
+ *   // Filter an array of items
+ *   TArray<FString> FilteredItems = MyFilter.Filter(Items);
+ */
+class FSlateIMStringFilter : public FSlateImTextFilter<const FString&>
+{
+public:
+	FSlateIMStringFilter()
+	{
+		// Create a TTextFilter for filtering FStrings
+		// The transform delegate simply returns the string itself in an array
+		TextFilter = MakeShared<TTextFilter<const FString&>>(
+			TTextFilter<const FString&>::FItemToStringArray::CreateLambda([](const FString& InString, TArray<FString>& OutStrings)
+			{
+				OutStrings.Add(InString);
+			})
+		);
+	}
+};
diff --git a/Engine/Plugins/Experimental/SlateIM/Source/SlateIM/Public/SlateIM.h b/Engine/Plugins/Experimental/SlateIM/Source/SlateIM/Public/SlateIM.h
--- a/Engine/Plugins/Experimental/SlateIM/Source/SlateIM/Public/SlateIM.h
+++ b/Engine/Plugins/Experimental/SlateIM/Source/SlateIM/Public/SlateIM.h
@@ -45,6 +45,7 @@ struct FTextBlockStyle;
 enum class ECheckBoxState : uint8;
 
 class SWidget;
+template<typename ItemType> class FSlateImTextFilter;
 
 
 class FSlateIMModule : public IModuleInterface
@@ -834,6 +835,24 @@ namespace SlateIM
 	 * @see CanUpdateSlateIM()
 	 */
 	SLATEIM_API EAppReturnType::Type ModalDialog(EAppMsgType::Type MessageType, const FStringView& DialogText, EAppMsgCategory Category = EAppMsgCategory::Warning, const FStringView& DialogTitle = FStringView());
+
+	/**
+	 * Text filter input box that updates the filter when text changes.
+	 *
+	 * @param InFilter The text filter to display and update
+	 * @param HintText Optional hint text to display when the input is empty
+	 * @return Whether the user changed or committed text this frame
+	 */
+	template<typename ItemType>
+	bool TextFilter(FSlateImTextFilter<ItemType>& InFilter, const FStringView& HintText = FStringView())
+	{
+		if (EditableText(InFilter.InputText, HintText) && InFilter.TextFilter.IsValid())
+		{
+			InFilter.TextFilter->SetRawFilterText(FText::FromString(InFilter.InputText));
+			return true;
+		}
+		return false;
+	}
 #pragma endregion Util
 
 #pragma region Queries
