SlateIM: All roots can now have root tab groups.

#Synced-CL 50961346

diff --git a/Engine/Plugins/Experimental/SlateIM/Source/SlateIM/Private/Containers/SImTabGroup.cpp b/Engine/Plugins/Experimental/SlateIM/Source/SlateIM/Private/Containers/SImTabGroup.cpp
--- a/Engine/Plugins/Experimental/SlateIM/Source/SlateIM/Private/Containers/SImTabGroup.cpp
+++ b/Engine/Plugins/Experimental/SlateIM/Source/SlateIM/Private/Containers/SImTabGroup.cpp
@@ -144,11 +144,9 @@ void SImTabGroup::Construct(const FArguments& InArgs)
 		TabManager = FGlobalTabmanager::Get()->NewTabManager(ParentTab.ToSharedRef());
 	}
 	else if (SlateIM::FSlateIMManager::Get().HasCurrentRoot()
-		&& SlateIM::FSlateIMManager::Get().GetCurrentRoot().RootWidget.IsValid()
-		&& SlateIM::FSlateIMManager::Get().GetCurrentRoot().RootWidget->IsA<FSlateIMWindowRoot>())
+		&& SlateIM::FSlateIMManager::Get().GetCurrentRoot().RootWidget.IsValid())
 	{
-		TSharedRef<FSlateIMWindowRoot> WindowRoot = StaticCastSharedPtr<FSlateIMWindowRoot>(SlateIM::FSlateIMManager::Get().GetCurrentRoot().RootWidget).ToSharedRef();
-		TabManager = FGlobalTabmanager::Get()->NewTabManager(WindowRoot->GetRootTab());
+		TabManager = FGlobalTabmanager::Get()->NewTabManager(SlateIM::FSlateIMManager::Get().GetCurrentRoot().RootWidget->GetRootTab());
 	}
 	else
 	{
diff --git a/Engine/Plugins/Experimental/SlateIM/Source/SlateIM/Private/Roots/ISlateIMRoot.h b/Engine/Plugins/Experimental/SlateIM/Source/SlateIM/Private/Roots/ISlateIMRoot.h
--- a/Engine/Plugins/Experimental/SlateIM/Source/SlateIM/Private/Roots/ISlateIMRoot.h
+++ b/Engine/Plugins/Experimental/SlateIM/Source/SlateIM/Private/Roots/ISlateIMRoot.h
@@ -9,6 +9,7 @@
 
 struct FSlateIMInputState;
 struct FSlateIMSlotData;
+class SDockTab;
 class SWidget;
 
 class ISlateIMRoot : public ISlateIMTypeChecking
@@ -18,4 +19,6 @@ public:
 	virtual bool IsVisible() const = 0;
 
 	virtual FSlateIMInputState& GetInputState() = 0;
+
+	virtual TSharedRef<SDockTab> GetRootTab() const = 0;
 };
diff --git a/Engine/Plugins/Experimental/SlateIM/Source/SlateIM/Private/Roots/SlateIMExposedRoot.cpp b/Engine/Plugins/Experimental/SlateIM/Source/SlateIM/Private/Roots/SlateIMExposedRoot.cpp
--- a/Engine/Plugins/Experimental/SlateIM/Source/SlateIM/Private/Roots/SlateIMExposedRoot.cpp
+++ b/Engine/Plugins/Experimental/SlateIM/Source/SlateIM/Private/Roots/SlateIMExposedRoot.cpp
@@ -3,11 +3,13 @@
 #include "SlateIMExposedRoot.h"
 
 #include "Misc/SlateIMSlotData.h"
+#include "Widgets/Docking/SDockTab.h"
 #include "Widgets/SImWrapper.h"
 
 
 FSlateIMExposedRoot::FSlateIMExposedRoot()
 	: ExposedWidget(SNew(SImWrapper))
+	, RootTab(SNew(SDockTab))
 {
 }
 
@@ -29,6 +31,11 @@ FSlateIMInputState& FSlateIMExposedRoot::GetInputState()
 	return ExposedWidget->InputState;
 }
 
+TSharedRef<SDockTab> FSlateIMExposedRoot::GetRootTab() const
+{
+	return RootTab;
+}
+
 TSharedRef<SWidget> FSlateIMExposedRoot::GetExposedWidget() const
 {
 	return ExposedWidget;
diff --git a/Engine/Plugins/Experimental/SlateIM/Source/SlateIM/Private/Roots/SlateIMExposedRoot.h b/Engine/Plugins/Experimental/SlateIM/Source/SlateIM/Private/Roots/SlateIMExposedRoot.h
--- a/Engine/Plugins/Experimental/SlateIM/Source/SlateIM/Private/Roots/SlateIMExposedRoot.h
+++ b/Engine/Plugins/Experimental/SlateIM/Source/SlateIM/Private/Roots/SlateIMExposedRoot.h
@@ -17,10 +17,11 @@ public:
 	virtual void UpdateChild(TSharedRef<SWidget> Child, const FSlateIMSlotData& AlignmentData) override;
 	virtual bool IsVisible() const override { return true; }
 	virtual FSlateIMInputState& GetInputState() override;
+	virtual TSharedRef<SDockTab> GetRootTab() const override;
 	
 	TSharedRef<SWidget> GetExposedWidget() const;
 	
 private:
 	TSharedRef<SImWrapper> ExposedWidget;
-
+	TSharedRef<SDockTab> RootTab;
 };
diff --git a/Engine/Plugins/Experimental/SlateIM/Source/SlateIM/Private/Roots/SlateIMViewportRoot.cpp b/Engine/Plugins/Experimental/SlateIM/Source/SlateIM/Private/Roots/SlateIMViewportRoot.cpp
--- a/Engine/Plugins/Experimental/SlateIM/Source/SlateIM/Private/Roots/SlateIMViewportRoot.cpp
+++ b/Engine/Plugins/Experimental/SlateIM/Source/SlateIM/Private/Roots/SlateIMViewportRoot.cpp
@@ -11,11 +11,13 @@
 #include "SlateIM.h"
 #include "Engine/GameViewportClient.h"
 #include "Engine/LocalPlayer.h"
+#include "Widgets/Docking/SDockTab.h"
 #include "Widgets/Layout/SScaleBox.h"
 #include "Widgets/SImWrapper.h"
 
 FSlateIMViewportRoot::FSlateIMViewportRoot(UGameViewportClient* InGameViewport)
 	: GameViewport(InGameViewport)
+	, RootTab(SNew(SDockTab))
 {
 	ensure(InGameViewport);
 }
@@ -23,6 +25,7 @@ FSlateIMViewportRoot::FSlateIMViewportRoot(UGameViewportClient* InGameViewport)
 FSlateIMViewportRoot::FSlateIMViewportRoot(ULocalPlayer* InLocalPlayer)
 	: GameViewport(InLocalPlayer ? InLocalPlayer->ViewportClient : nullptr)
 	, LocalPlayer(InLocalPlayer)
+	, RootTab(SNew(SDockTab))
 {
 	ensure(InLocalPlayer);
 	ensure(GameViewport.IsValid());
@@ -31,6 +34,7 @@ FSlateIMViewportRoot::FSlateIMViewportRoot(ULocalPlayer* InLocalPlayer)
 #if WITH_EDITOR
 FSlateIMViewportRoot::FSlateIMViewportRoot(TSharedPtr<IAssetViewport> InAssetViewport)
 	: AssetViewport(InAssetViewport)
+	, RootTab(SNew(SDockTab))
 {
 	ensure(InAssetViewport);
 }
@@ -75,6 +79,11 @@ FSlateIMInputState& FSlateIMViewportRoot::GetInputState()
 	return Wrapper->InputState;
 }
 
+TSharedRef<SDockTab> FSlateIMViewportRoot::GetRootTab() const
+{
+	return RootTab;
+}
+
 void FSlateIMViewportRoot::UpdateViewport(const SlateIM::FViewportRootLayout& Layout)
 {
 	if (!ViewportRoot)
diff --git a/Engine/Plugins/Experimental/SlateIM/Source/SlateIM/Private/Roots/SlateIMViewportRoot.h b/Engine/Plugins/Experimental/SlateIM/Source/SlateIM/Private/Roots/SlateIMViewportRoot.h
--- a/Engine/Plugins/Experimental/SlateIM/Source/SlateIM/Private/Roots/SlateIMViewportRoot.h
+++ b/Engine/Plugins/Experimental/SlateIM/Source/SlateIM/Private/Roots/SlateIMViewportRoot.h
@@ -40,6 +40,7 @@ public:
 		return true;
 	}
 	virtual FSlateIMInputState& GetInputState() override;
+	virtual TSharedRef<SDockTab> GetRootTab() const override;
 
 	void UpdateViewport(const SlateIM::FViewportRootLayout& Layout);
 
@@ -54,5 +55,6 @@ private:
 	TSharedPtr<SImWrapper> Wrapper;
 	TSharedPtr<SWidget> ViewportRoot;
 	SConstraintCanvas::FSlot* Slot = nullptr;
+	TSharedRef<SDockTab> RootTab;
 };
 #endif // WITH_ENGINE
diff --git a/Engine/Plugins/Experimental/SlateIM/Source/SlateIM/Private/Roots/SlateIMWindowRoot.h b/Engine/Plugins/Experimental/SlateIM/Source/SlateIM/Private/Roots/SlateIMWindowRoot.h
--- a/Engine/Plugins/Experimental/SlateIM/Source/SlateIM/Private/Roots/SlateIMWindowRoot.h
+++ b/Engine/Plugins/Experimental/SlateIM/Source/SlateIM/Private/Roots/SlateIMWindowRoot.h
@@ -19,12 +19,11 @@ public:
 	virtual void UpdateChild(TSharedRef<SWidget> Child, const FSlateIMSlotData& AlignmentData) override;
 	virtual bool IsVisible() const override;
 	virtual FSlateIMInputState& GetInputState() override;
+	virtual TSharedRef<SDockTab> GetRootTab() const override;
 
 	void UpdateWindow(const FStringView& Title);
 
 	TSharedPtr<SWindow> GetWindow() const;
-
-	TSharedRef<SDockTab> GetRootTab() const;
 	
 private:
 	TWeakPtr<SWindow> RootWindow;
