[SlateIM] make sure to update a combo box selected text if we change it via code

#Synced-CL 51050165

#Synced-CL 51050165

diff --git a/Engine/Plugins/Experimental/SlateIM/Source/SlateIM/Private/Implementation/SlateIM_Widgets.cpp b/Engine/Plugins/Experimental/SlateIM/Source/SlateIM/Private/Implementation/SlateIM_Widgets.cpp
--- a/Engine/Plugins/Experimental/SlateIM/Source/SlateIM/Private/Implementation/SlateIM_Widgets.cpp
+++ b/Engine/Plugins/Experimental/SlateIM/Source/SlateIM/Private/Implementation/SlateIM_Widgets.cpp
@@ -641,6 +641,7 @@ namespace SlateIM
 		TSharedPtr<STextComboBox> ComboWidget = Scope.GetWidget();
 
 		Scope.HashData(InComboStyle);
+		Scope.HashData(InOutSelectedItemIndex);
 		
 		using FComboBoxData = FSlateIMDataStore<TArray<TSharedPtr<FString>>>;
 		
@@ -704,6 +705,13 @@ namespace SlateIM
 
 			if (Scope.IsDataHashDirty())
 			{
+				TSharedPtr<FString> SelectedItem = ComboWidget->GetSelectedItem();
+				const int32 SelectedIndex = SelectedItem.IsValid() ? ComboBoxData->Data.IndexOfByKey(SelectedItem) : INDEX_NONE;
+				if (SelectedIndex != InOutSelectedItemIndex && ComboBoxData->Data.IsValidIndex(InOutSelectedItemIndex))
+				{
+					ComboWidget->SetSelectedItem(ComboBoxData->Data[InOutSelectedItemIndex]);
+				}
+
 				ComboWidget->SetStyle(InComboStyle);
 			}
 		}
