[SlateIM] replaced check by error reporting when the tab couldn't be invoked

#rnx
#Synced-CL 51296367

diff --git a/Engine/Plugins/Experimental/SlateIM/Source/SlateIM/Private/SlateIMWidgetBase.cpp b/Engine/Plugins/Experimental/SlateIM/Source/SlateIM/Private/SlateIMWidgetBase.cpp
--- a/Engine/Plugins/Experimental/SlateIM/Source/SlateIM/Private/SlateIMWidgetBase.cpp
+++ b/Engine/Plugins/Experimental/SlateIM/Source/SlateIM/Private/SlateIMWidgetBase.cpp
@@ -7,6 +7,7 @@
 #include "Framework/Docking/SDockingTabStack.h"
 #include "Framework/Docking/TabManager.h"
 #include "SlateIM.h"
+#include "Misc/SlateIMLogging.h"
 #include "Widgets/Docking/SDockTab.h"
 #include "Widgets/SNullWidget.h"
 
@@ -126,15 +127,20 @@ void FSlateIMNomadTabBase::DrawWidget(float DeltaTime)
 	{
 		return;
 	}
-	
+
 	TSharedPtr<SDockTab> RootTabPin = RootTab.Pin();
 	if (!RootTabPin.IsValid())
 	{
 		// This is the first time we call DrawWidget() since being enabled - invoke the tab
 		RootTab = RootTabPin = FGlobalTabmanager::Get()->TryInvokeTab(GetWidgetName());
-		check(RootTabPin.IsValid())
+		if (!RootTabPin.IsValid())
+		{
+			UE_LOGF(LogSlateIM, Error, "Cannot invoke tab '%ls'; disabling widget", *GetWidgetName().ToString());
+			DisableWidget();
+			return;
+		}
 	}
-		
+
 	TSharedPtr<SWidget> NewExposedWidget;
 	if (SlateIM::BeginExposedRoot(GetWidgetName(), NewExposedWidget))
 	{
