SlateIM: Initial expansion state now works correctly for table rows.

#Synced-CL 51425194

diff --git a/Engine/Plugins/Experimental/SlateIM/Source/SlateIM/Private/Implementation/SlateIM_Containers.cpp b/Engine/Plugins/Experimental/SlateIM/Source/SlateIM/Private/Implementation/SlateIM_Containers.cpp
--- a/Engine/Plugins/Experimental/SlateIM/Source/SlateIM/Private/Implementation/SlateIM_Containers.cpp
+++ b/Engine/Plugins/Experimental/SlateIM/Source/SlateIM/Private/Implementation/SlateIM_Containers.cpp
@@ -665,8 +665,9 @@ namespace SlateIM
 		{
 			if (TSharedPtr<SImTableView> OwningTable = ParentRow->GetOwningTable())
 			{
-				bool bHasExpansionState = false;
-				bool bExpanded = false;
+				// If the default expanded state is true, we have a state because the default is unexpanded.
+				bool bHasExpansionState = bDefaultExpanded;
+				bool bExpanded = bDefaultExpanded;
 				TSharedPtr<ISlateIMContainer> GrandParent = ParentRow->GetParent();
 
 				if (GrandParent->IsA<FSlateIMTableBody>())
