[SlateIM] add simple API for button that can be disabled

#Synced-CL 52255093

diff --git a/Engine/Plugins/Experimental/SlateIM/Source/SlateIM/Private/Implementation/SlateIM_Widgets.cpp b/Engine/Plugins/Experimental/SlateIM/Source/SlateIM/Private/Implementation/SlateIM_Widgets.cpp
--- a/Engine/Plugins/Experimental/SlateIM/Source/SlateIM/Private/Implementation/SlateIM_Widgets.cpp
+++ b/Engine/Plugins/Experimental/SlateIM/Source/SlateIM/Private/Implementation/SlateIM_Widgets.cpp
@@ -302,6 +302,23 @@ namespace SlateIM
 		return bWasClicked;
 	}
 
+	bool Button(const FStringView& InText, const bool bEnabled, const FButtonStyle* InStyle)
+	{
+		if (!bEnabled)
+		{
+			BeginDisabledState();
+		}
+
+		const bool bResult = Button(InText, InStyle);
+
+		if (!bEnabled)
+		{
+			EndDisabledState();
+		}
+
+		return bResult;
+	}
+
 	bool CheckBox(const FStringView& InText, bool& InOutCurrentState, const FCheckBoxStyle* CheckBoxStyle)
 	{
 		ECheckBoxState CurrentEnumState = InOutCurrentState ? ECheckBoxState::Checked : ECheckBoxState::Unchecked;
diff --git a/Engine/Plugins/Experimental/SlateIM/Source/SlateIM/Private/Misc/SlateIMExamples.cpp b/Engine/Plugins/Experimental/SlateIM/Source/SlateIM/Private/Misc/SlateIMExamples.cpp
--- a/Engine/Plugins/Experimental/SlateIM/Source/SlateIM/Private/Misc/SlateIMExamples.cpp
+++ b/Engine/Plugins/Experimental/SlateIM/Source/SlateIM/Private/Misc/SlateIMExamples.cpp
@@ -204,6 +204,13 @@ void FSlateIMTestWidget::DrawBasics()
 				UE_LOG(LogSlateIM, Log, TEXT("Button was clicked"));
 			}
 
+			{
+				SlateIM::BeginHorizontalStack();
+				SlateIM::Button(TEXT("Button you can disable"), bIsButtonEnabled);
+				SlateIM::CheckBox(TEXT("Enabled"), bIsButtonEnabled);
+				SlateIM::EndHorizontalStack();
+			}
+
 			{
 				SlateIM::BeginVerticalStack();
 				// EditableText Example
diff --git a/Engine/Plugins/Experimental/SlateIM/Source/SlateIM/Private/Misc/SlateIMExamples.h b/Engine/Plugins/Experimental/SlateIM/Source/SlateIM/Private/Misc/SlateIMExamples.h
--- a/Engine/Plugins/Experimental/SlateIM/Source/SlateIM/Private/Misc/SlateIMExamples.h
+++ b/Engine/Plugins/Experimental/SlateIM/Source/SlateIM/Private/Misc/SlateIMExamples.h
@@ -48,6 +48,7 @@ private:
 	bool bShouldAddChildRow = false;
 	TArray<FString> ComboBoxItems = { TEXT("Option 1"), TEXT("Option 2"), TEXT("Option 3"), TEXT("Option 4") };
 	bool bRefreshComboItems = false;
+	bool bIsButtonEnabled = true;
 	int32 DynamicTabCount = 3;
 	FSlateIMStringFilter FruitFilter;
 	TOptional<TArray<FString>> FilteredFruitList;
diff --git a/Engine/Plugins/Experimental/SlateIM/Source/SlateIM/Public/SlateIM.h b/Engine/Plugins/Experimental/SlateIM/Source/SlateIM/Public/SlateIM.h
--- a/Engine/Plugins/Experimental/SlateIM/Source/SlateIM/Public/SlateIM.h
+++ b/Engine/Plugins/Experimental/SlateIM/Source/SlateIM/Public/SlateIM.h
@@ -521,6 +521,18 @@ namespace SlateIM
 	 */
 	SLATEIM_API bool Button(const FStringView& InText, const FButtonStyle* InStyle = nullptr);
 
+	/**
+	 * Display a button with text
+	 * 
+	 * @param InText The text to display on the button
+	 * @param bEnabled Indicate if the button is enabled or not (grayed)
+	 * @param InStyle (optional) A Button Style to override the default style. Changes to the style after the first invocation are not reflected,
+	 *				  provide a new style object to update the visuals at runtime
+	 *					
+	 * @return Whether the user clicked the button this frame
+	 */
+	SLATEIM_API bool Button(const FStringView& InText, const bool bEnabled, const FButtonStyle* InStyle = nullptr);
+
 	/**
 	 * Display a two-state checkbox
 	 * 
