SlateIM: Major changes
- Change to struct parameters to the majority of methods.
- Exported canvas functions
- Added Maximize option to Stack and Wrap calls.
- Updated examples file to use the struct params.
- Added BP Support
- Fixed compilation in non-engine configuration

#Synced-CL 52294106

diff --git a/Engine/Plugins/Experimental/SlateIM/SlateIM.uplugin b/Engine/Plugins/Experimental/SlateIM/SlateIM.uplugin
--- a/Engine/Plugins/Experimental/SlateIM/SlateIM.uplugin
+++ b/Engine/Plugins/Experimental/SlateIM/SlateIM.uplugin
@@ -21,10 +21,20 @@
 			"Type": "Runtime",
 			"LoadingPhase": "Default"
 		},
+		{
+			"Name": "SlateIMEngine",
+			"Type": "Runtime",
+			"LoadingPhase": "Default"
+		},
 		{
 			"Name": "SlateIMInGame",
 			"Type": "Runtime",
 			"LoadingPhase": "Default"
+		},
+		{
+			"Name": "SlateIMBlueprint",
+			"Type": "Runtime",
+			"LoadingPhase": "Default"
 		}
 	]
 }
\ No newline at end of file
diff --git a/Engine/Plugins/Experimental/SlateIM/Source/SlateIM/Private/Implementation/SlateIMCanvasDrawCommandList.cpp b/Engine/Plugins/Experimental/SlateIM/Source/SlateIM/Private/Implementation/SlateIMCanvasDrawCommandList.cpp
deleted file mode 100644
diff --git a/Engine/Plugins/Experimental/SlateIM/Source/SlateIM/Private/Implementation/SlateIMCanvasDrawCommandList.h b/Engine/Plugins/Experimental/SlateIM/Source/SlateIM/Private/Implementation/SlateIMCanvasDrawCommandList.h
deleted file mode 100644
diff --git a/Engine/Plugins/Experimental/SlateIM/Source/SlateIM/Private/Implementation/SlateIM_Canvas.cpp b/Engine/Plugins/Experimental/SlateIM/Source/SlateIM/Private/Implementation/SlateIM_Canvas.cpp
--- a/Engine/Plugins/Experimental/SlateIM/Source/SlateIM/Private/Implementation/SlateIM_Canvas.cpp
+++ b/Engine/Plugins/Experimental/SlateIM/Source/SlateIM/Private/Implementation/SlateIM_Canvas.cpp
@@ -11,7 +11,7 @@
 #include "Engine/CanvasRenderTarget2D.h"
 #include "Engine/Texture2D.h"
 #include "Engine/TextureRenderTarget2D.h"
-#include "Implementation/SlateIMCanvasDrawCommandList.h"
+#include "Implementation/SlateIMEngineCanvasDrawCommandList.h"
 #include "Math/IntPoint.h"
 #include "Misc/SlateIMManager.h"
 #include "Misc/SlateIMSlateResources.h"
@@ -72,29 +72,29 @@ namespace SlateIM::Canvas
 		}
 
 		// Create uobject proxy for queueing draw calls
-		USlateIMCanvasDrawCommandList* UpdateContainer = nullptr;
+		USlateIMEngineCanvasDrawCommandList* UpdateContainer = nullptr;
 		TSharedPtr<FSlateIMPinnedObjectResourceDataStore> UpdateResource = CanvasImage->GetMetaData<FSlateIMPinnedObjectResourceDataStore>();
 
 		if (!UpdateResource)
 		{
-			UpdateContainer = NewObject<USlateIMCanvasDrawCommandList>();
+			UpdateContainer = NewObject<USlateIMEngineCanvasDrawCommandList>();
 			UpdateResource = MakeShared<FSlateIMPinnedObjectResourceDataStore>(FPinnedObjectResource{ TStrongObjectPtr<UObject>(UpdateContainer) });
 			CanvasImage->AddMetadata(UpdateResource.ToSharedRef());
 		}
 		else
 		{
-			UpdateContainer = CastChecked<USlateIMCanvasDrawCommandList>(UpdateResource->Data.PinnedResource.Get());
+			UpdateContainer = CastChecked<USlateIMEngineCanvasDrawCommandList>(UpdateResource->Data.PinnedResource.Get());
 		}
 
 		CanvasRenderTarget->OnCanvasRenderTargetUpdate.Clear();
-		CanvasRenderTarget->OnCanvasRenderTargetUpdate.AddDynamic(UpdateContainer, &USlateIMCanvasDrawCommandList::ProcessCommands);
+		CanvasRenderTarget->OnCanvasRenderTargetUpdate.AddDynamic(UpdateContainer, &USlateIMEngineCanvasDrawCommandList::ProcessCommands);
 
 		// Defer update call to EndCanvas
 
 		return &CanvasResource->Data.Brush;
 	}
 
-	USlateIMCanvasDrawCommandList* GetCanvasCommandList()
+	USlateIMEngineCanvasDrawCommandList* GetCanvasCommandList()
 	{
 		TSharedPtr<SImCompoundWidget> CanvasContainer = FSlateIMManager::Get().GetCurrentContainer<SImCompoundWidget>();
 
@@ -104,7 +104,7 @@ namespace SlateIM::Canvas
 			{
 				if (TSharedPtr<FSlateIMPinnedObjectResourceDataStore> CanvasResource = CanvasImage->GetMetaData<FSlateIMPinnedObjectResourceDataStore>())
 				{
-					return CastChecked<USlateIMCanvasDrawCommandList>(CanvasResource->Data.PinnedResource.Get());
+					return CastChecked<USlateIMEngineCanvasDrawCommandList>(CanvasResource->Data.PinnedResource.Get());
 				}
 			}
 		}
@@ -130,9 +130,9 @@ namespace SlateIM::Canvas
 				return GetCanvasBrush(ImageWidget, WorldContextWeak, Width, Height);
 			};
 
-		Image_Internal(CanvasParams.ColorAndOpacity, CanvasParams.DesiredSize, GetBrush);
+		Image_Internal({.ColorAndOpacity = CanvasParams.ColorAndOpacity, .DesiredSize = CanvasParams.DesiredSize}, GetBrush);
 
-		if (USlateIMCanvasDrawCommandList* CommandList = GetCanvasCommandList())
+		if (USlateIMEngineCanvasDrawCommandList* CommandList = GetCanvasCommandList())
 		{
 			CommandList->SetUpdateType(CanvasParams.UpdateType);
 		}
@@ -140,7 +140,7 @@ namespace SlateIM::Canvas
 
 	void EndCanvas()
 	{
-		if (USlateIMCanvasDrawCommandList* CommandList = GetCanvasCommandList())
+		if (USlateIMEngineCanvasDrawCommandList* CommandList = GetCanvasCommandList())
 		{
 			if (!CommandList->NeedsUpdate())
 			{
@@ -182,7 +182,7 @@ namespace SlateIM::Canvas
 
 	void Invalidate()
 	{
-		if (USlateIMCanvasDrawCommandList* CommandList = GetCanvasCommandList())
+		if (USlateIMEngineCanvasDrawCommandList* CommandList = GetCanvasCommandList())
 		{
 			CommandList->Invalidate();
 		}
@@ -190,9 +190,9 @@ namespace SlateIM::Canvas
 
 	void SetClip(const FVector2f& ClipPosition)
 	{
-		if (USlateIMCanvasDrawCommandList* CommandList = GetCanvasCommandList())
+		if (USlateIMEngineCanvasDrawCommandList* CommandList = GetCanvasCommandList())
 		{
-			CommandList->EnqueueCommand(USlateIMCanvasDrawCommandList::FCanvasDrawCommand::CreateLambda(
+			CommandList->EnqueueCommand(USlateIMEngineCanvasDrawCommandList::FCanvasDrawCommand::CreateLambda(
 				[ClipPosition](UCanvas* InCanvas)
 				{
 					InCanvas->SetClip(ClipPosition.X, ClipPosition.Y);
@@ -203,9 +203,9 @@ namespace SlateIM::Canvas
 
 	void SetDrawColor(const FLinearColor& Color)
 	{
-		if (USlateIMCanvasDrawCommandList* CommandList = GetCanvasCommandList())
+		if (USlateIMEngineCanvasDrawCommandList* CommandList = GetCanvasCommandList())
 		{
-			CommandList->EnqueueCommand(USlateIMCanvasDrawCommandList::FCanvasDrawCommand::CreateLambda(
+			CommandList->EnqueueCommand(USlateIMEngineCanvasDrawCommandList::FCanvasDrawCommand::CreateLambda(
 				[Color](UCanvas* InCanvas)
 				{
 					InCanvas->SetLinearDrawColor(Color);
@@ -216,9 +216,9 @@ namespace SlateIM::Canvas
 
 	void SetDrawColor(const FColor& Color)
 	{
-		if (USlateIMCanvasDrawCommandList* CommandList = GetCanvasCommandList())
+		if (USlateIMEngineCanvasDrawCommandList* CommandList = GetCanvasCommandList())
 		{
-			CommandList->EnqueueCommand(USlateIMCanvasDrawCommandList::FCanvasDrawCommand::CreateLambda(
+			CommandList->EnqueueCommand(USlateIMEngineCanvasDrawCommandList::FCanvasDrawCommand::CreateLambda(
 				[Color](UCanvas* InCanvas)
 				{
 					InCanvas->SetDrawColor(Color);
@@ -229,9 +229,9 @@ namespace SlateIM::Canvas
 
 	void DrawItem(TSharedRef<FCanvasItem> ItemPtr, const FVector2f& CanvasPosition)
 	{
-		if (USlateIMCanvasDrawCommandList* CommandList = GetCanvasCommandList())
+		if (USlateIMEngineCanvasDrawCommandList* CommandList = GetCanvasCommandList())
 		{
-			CommandList->EnqueueCommand(USlateIMCanvasDrawCommandList::FCanvasDrawCommand::CreateLambda(
+			CommandList->EnqueueCommand(USlateIMEngineCanvasDrawCommandList::FCanvasDrawCommand::CreateLambda(
 				[ItemPtr, CanvasPosition](UCanvas* InCanvas)
 				{
 					InCanvas->DrawItem(*ItemPtr, CanvasPosition.X, CanvasPosition.Y);
@@ -242,9 +242,9 @@ namespace SlateIM::Canvas
 
 	void DrawIcon(const FCanvasIcon& Icon, const FVector2f& CanvasPosition, const FVector2f& Scale)
 	{
-		if (USlateIMCanvasDrawCommandList* CommandList = GetCanvasCommandList())
+		if (USlateIMEngineCanvasDrawCommandList* CommandList = GetCanvasCommandList())
 		{
-			CommandList->EnqueueCommand(USlateIMCanvasDrawCommandList::FCanvasDrawCommand::CreateLambda(
+			CommandList->EnqueueCommand(USlateIMEngineCanvasDrawCommandList::FCanvasDrawCommand::CreateLambda(
 				[Icon, CanvasPosition, Scale](UCanvas* InCanvas)
 				{
 					InCanvas->DrawScaledIcon(
@@ -258,11 +258,11 @@ namespace SlateIM::Canvas
 		}
 	}
 
-	void DrawLine(const FVector2D& CanvasPositionA, const FVector2D& CanvasPositionB, float Thickness, FLinearColor RenderColor)
+	void DrawLine(const FVector2D& CanvasPositionA, const FVector2D& CanvasPositionB, float Thickness, const FLinearColor& RenderColor)
 	{
-		if (USlateIMCanvasDrawCommandList* CommandList = GetCanvasCommandList())
+		if (USlateIMEngineCanvasDrawCommandList* CommandList = GetCanvasCommandList())
 		{
-			CommandList->EnqueueCommand(USlateIMCanvasDrawCommandList::FCanvasDrawCommand::CreateLambda(
+			CommandList->EnqueueCommand(USlateIMEngineCanvasDrawCommandList::FCanvasDrawCommand::CreateLambda(
 				[CanvasPositionA, CanvasPositionB, Thickness, RenderColor](UCanvas* InCanvas)
 				{
 					InCanvas->K2_DrawLine(
@@ -278,9 +278,9 @@ namespace SlateIM::Canvas
 
 	void DrawTexture(UTexture* RenderTexture, const FVector2D& CanvasPosition, const FVector2D& CanvasSize, const FTileRenderParams& RenderParams)
 	{
-		if (USlateIMCanvasDrawCommandList* CommandList = GetCanvasCommandList())
+		if (USlateIMEngineCanvasDrawCommandList* CommandList = GetCanvasCommandList())
 		{
-			CommandList->EnqueueCommand(USlateIMCanvasDrawCommandList::FCanvasDrawCommand::CreateLambda(
+			CommandList->EnqueueCommand(USlateIMEngineCanvasDrawCommandList::FCanvasDrawCommand::CreateLambda(
 				[RenderTexture, CanvasPosition, CanvasSize, RenderParams](UCanvas* InCanvas)
 				{
 					InCanvas->K2_DrawTexture(
@@ -301,9 +301,9 @@ namespace SlateIM::Canvas
 
 	void DrawMaterial(UMaterialInterface* RenderMaterial, const FVector2D& CanvasPosition, const FVector2D& CanvasSize, const FTileRenderParams& RenderParams)
 	{
-		if (USlateIMCanvasDrawCommandList* CommandList = GetCanvasCommandList())
+		if (USlateIMEngineCanvasDrawCommandList* CommandList = GetCanvasCommandList())
 		{
-			CommandList->EnqueueCommand(USlateIMCanvasDrawCommandList::FCanvasDrawCommand::CreateLambda(
+			CommandList->EnqueueCommand(USlateIMEngineCanvasDrawCommandList::FCanvasDrawCommand::CreateLambda(
 				[RenderMaterial, CanvasPosition, CanvasSize, RenderParams](UCanvas* InCanvas)
 				{
 					InCanvas->K2_DrawMaterial(
@@ -322,9 +322,9 @@ namespace SlateIM::Canvas
 
 	void DrawText(UFont* RenderFont, const FString& RenderText, const FVector2D& CanvasPosition, const FTextRenderParams& RenderParams)
 	{
-		if (USlateIMCanvasDrawCommandList* CommandList = GetCanvasCommandList())
+		if (USlateIMEngineCanvasDrawCommandList* CommandList = GetCanvasCommandList())
 		{
-			CommandList->EnqueueCommand(USlateIMCanvasDrawCommandList::FCanvasDrawCommand::CreateLambda(
+			CommandList->EnqueueCommand(USlateIMEngineCanvasDrawCommandList::FCanvasDrawCommand::CreateLambda(
 				[RenderFont, RenderText, CanvasPosition, RenderParams](UCanvas* InCanvas)
 				{
 					InCanvas->K2_DrawText(
@@ -348,9 +348,9 @@ namespace SlateIM::Canvas
 
 	void DrawText(UFont* RenderFont, const FString& RenderText, const FVector2f& CanvasPosition, const FVector2f& Scale, const FFontRenderInfo& RenderInfo)
 	{
-		if (USlateIMCanvasDrawCommandList* CommandList = GetCanvasCommandList())
+		if (USlateIMEngineCanvasDrawCommandList* CommandList = GetCanvasCommandList())
 		{
-			CommandList->EnqueueCommand(USlateIMCanvasDrawCommandList::FCanvasDrawCommand::CreateLambda(
+			CommandList->EnqueueCommand(USlateIMEngineCanvasDrawCommandList::FCanvasDrawCommand::CreateLambda(
 				[RenderFont, RenderText, CanvasPosition, Scale, RenderInfo](UCanvas* InCanvas)
 				{
 					InCanvas->DrawText(
@@ -369,9 +369,9 @@ namespace SlateIM::Canvas
 
 	void DrawWrappedText(const UFont* RenderFont, const FString& RenderText, const FVector2f& CanvasPosition, bool bCenterTextX, bool bCenterTextY, const FVector2f& Scale, const FFontRenderInfo& RenderInfo)
 	{
-		if (USlateIMCanvasDrawCommandList* CommandList = GetCanvasCommandList())
+		if (USlateIMEngineCanvasDrawCommandList* CommandList = GetCanvasCommandList())
 		{
-			CommandList->EnqueueCommand(USlateIMCanvasDrawCommandList::FCanvasDrawCommand::CreateLambda(
+			CommandList->EnqueueCommand(USlateIMEngineCanvasDrawCommandList::FCanvasDrawCommand::CreateLambda(
 				[RenderFont, RenderText, CanvasPosition, Scale, bCenterTextX, bCenterTextY, RenderInfo](UCanvas* InCanvas)
 				{
 					int32 outXL, outYL;
@@ -397,9 +397,9 @@ namespace SlateIM::Canvas
 
 	void DrawBorder(UTexture* BorderTexture, UTexture* BackgroundTexture, UTexture* LeftBorderTexture, UTexture* RightBorderTexture, UTexture* TopBorderTexture, UTexture* BottomBorderTexture, const FVector2D& CanvasPosition, const FVector2D& CanvasSize, const FBorderRenderParams& RenderParams)
 	{
-		if (USlateIMCanvasDrawCommandList* CommandList = GetCanvasCommandList())
+		if (USlateIMEngineCanvasDrawCommandList* CommandList = GetCanvasCommandList())
 		{
-			CommandList->EnqueueCommand(USlateIMCanvasDrawCommandList::FCanvasDrawCommand::CreateLambda(
+			CommandList->EnqueueCommand(USlateIMEngineCanvasDrawCommandList::FCanvasDrawCommand::CreateLambda(
 				[BorderTexture,BackgroundTexture, LeftBorderTexture, RightBorderTexture, TopBorderTexture, BottomBorderTexture, 
 					CanvasPosition, CanvasSize, RenderParams](UCanvas* InCanvas)
 				{
@@ -428,9 +428,9 @@ namespace SlateIM::Canvas
 
 	void DrawBox(const FVector2D& CanvasPosition, const FVector2D& CanvasSize, float Thickness, FLinearColor RenderColor)
 	{
-		if (USlateIMCanvasDrawCommandList* CommandList = GetCanvasCommandList())
+		if (USlateIMEngineCanvasDrawCommandList* CommandList = GetCanvasCommandList())
 		{
-			CommandList->EnqueueCommand(USlateIMCanvasDrawCommandList::FCanvasDrawCommand::CreateLambda(
+			CommandList->EnqueueCommand(USlateIMEngineCanvasDrawCommandList::FCanvasDrawCommand::CreateLambda(
 				[CanvasPosition, CanvasSize, Thickness, RenderColor](UCanvas* InCanvas)
 				{
 					InCanvas->K2_DrawBox(
@@ -446,9 +446,9 @@ namespace SlateIM::Canvas
 
 	void DrawTriangles(UTexture* RenderTexture, TSharedRef<TArray<FCanvasUVTri>> Triangles)
 	{
-		if (USlateIMCanvasDrawCommandList* CommandList = GetCanvasCommandList())
+		if (USlateIMEngineCanvasDrawCommandList* CommandList = GetCanvasCommandList())
 		{
-			CommandList->EnqueueCommand(USlateIMCanvasDrawCommandList::FCanvasDrawCommand::CreateLambda(
+			CommandList->EnqueueCommand(USlateIMEngineCanvasDrawCommandList::FCanvasDrawCommand::CreateLambda(
 				[RenderTexture, Triangles](UCanvas* InCanvas)
 				{
 					InCanvas->K2_DrawTriangle(RenderTexture, *Triangles);
@@ -459,9 +459,9 @@ namespace SlateIM::Canvas
 
 	void DrawMaterialTriangles(UMaterialInterface* RenderMaterial, TSharedRef<TArray<FCanvasUVTri>> Triangles)
 	{
-		if (USlateIMCanvasDrawCommandList* CommandList = GetCanvasCommandList())
+		if (USlateIMEngineCanvasDrawCommandList* CommandList = GetCanvasCommandList())
 		{
-			CommandList->EnqueueCommand(USlateIMCanvasDrawCommandList::FCanvasDrawCommand::CreateLambda(
+			CommandList->EnqueueCommand(USlateIMEngineCanvasDrawCommandList::FCanvasDrawCommand::CreateLambda(
 				[RenderMaterial, Triangles](UCanvas* InCanvas)
 				{
 					InCanvas->K2_DrawMaterialTriangle(RenderMaterial, *Triangles);
@@ -472,9 +472,9 @@ namespace SlateIM::Canvas
 
 	void DrawPolygon(UTexture* RenderTexture, const FVector2D& CanvasPosition, const FVector2D& Radius, int32 NumberOfSides, FLinearColor RenderColor)
 	{
-		if (USlateIMCanvasDrawCommandList* CommandList = GetCanvasCommandList())
+		if (USlateIMEngineCanvasDrawCommandList* CommandList = GetCanvasCommandList())
 		{
-			CommandList->EnqueueCommand(USlateIMCanvasDrawCommandList::FCanvasDrawCommand::CreateLambda(
+			CommandList->EnqueueCommand(USlateIMEngineCanvasDrawCommandList::FCanvasDrawCommand::CreateLambda(
 				[RenderTexture, CanvasPosition, Radius, NumberOfSides, RenderColor](UCanvas* InCanvas)
 				{
 					InCanvas->K2_DrawPolygon(
diff --git a/Engine/Plugins/Experimental/SlateIM/Source/SlateIM/Private/Implementation/SlateIM_Containers.cpp b/Engine/Plugins/Experimental/SlateIM/Source/SlateIM/Private/Implementation/SlateIM_Containers.cpp
--- a/Engine/Plugins/Experimental/SlateIM/Source/SlateIM/Private/Implementation/SlateIM_Containers.cpp
+++ b/Engine/Plugins/Experimental/SlateIM/Source/SlateIM/Private/Implementation/SlateIM_Containers.cpp
@@ -94,8 +94,13 @@ namespace SlateIM
 		FSlateIMManager::Get().PopContainer<SImWrapBox>();
 	}
 
-	void BeginHorizontalStack()
+	void BeginHorizontalStack(bool bMaximizeContent)
 	{
+		if (bMaximizeContent)
+		{
+			SlateIM::Maximize();
+		}
+
 		BeginStack(Orient_Horizontal);
 	}
 
@@ -104,8 +109,13 @@ namespace SlateIM
 		EndStack();
 	}
 
-	void BeginVerticalStack()
+	void BeginVerticalStack(bool bMaximizeContent)
 	{
+		if (bMaximizeContent)
+		{
+			SlateIM::Maximize();
+		}
+
 		BeginStack(Orient_Vertical);
 	}
 
@@ -114,8 +124,13 @@ namespace SlateIM
 		EndStack();
 	}
 
-	void BeginHorizontalWrap()
+	void BeginHorizontalWrap(bool bMaximizeContent)
 	{
+		if (bMaximizeContent)
+		{
+			SlateIM::Maximize();
+		}
+
 		BeginWrap(Orient_Horizontal);
 	}
 
@@ -124,8 +139,13 @@ namespace SlateIM
 		EndWrap();
 	}
 
-	void BeginVerticalWrap()
+	void BeginVerticalWrap(bool bMaximizeContent)
 	{
+		if (bMaximizeContent)
+		{
+			SlateIM::Maximize();
+		}
+
 		BeginWrap(Orient_Vertical);
 	}
 
@@ -135,12 +155,17 @@ namespace SlateIM
 	}
 
 	void BeginBorder(const FSlateBrush* BackgroundImage, EOrientation Orientation, bool bAbsorbMouse, FMargin ContentPadding)
+	{
+		BeginBorder(BackgroundImage, {.Orientation = Orientation, .bAbsorbMouse = bAbsorbMouse, .ContentPadding = ContentPadding});
+	}
+
+	void BeginBorder(const FSlateBrush* BackgroundImage, const FBorderParams& Params)
 	{
 		TSharedPtr<SImCompoundWidget> ContainerWidget;
 
 		// Note widget scope must fall out of scope before we push container or else we are pushing this new container into itself instead of into its parent
 		{
-			FWidgetScope<SImCompoundWidget> Scope(ContentPadding, HAlign_Fill, VAlign_Fill, false);
+			FWidgetScope<SImCompoundWidget> Scope(Params.ContentPadding, HAlign_Fill, VAlign_Fill, false);
 			ContainerWidget = Scope.GetWidget();
 
 			if (!ContainerWidget)
@@ -152,9 +177,9 @@ namespace SlateIM
 
 
 			ContainerWidget->SetBackgroundImage(BackgroundImage);
-			ContainerWidget->SetContentPadding(ContentPadding);
-			ContainerWidget->SetAbsorbMouse(bAbsorbMouse);
-			ContainerWidget->SetOrientation(Orientation);
+			ContainerWidget->SetContentPadding(Params.ContentPadding);
+			ContainerWidget->SetAbsorbMouse(Params.bAbsorbMouse);
+			ContainerWidget->SetOrientation(Params.Orientation);
 		}
 
 		FSlateIMManager::Get().PushContainer(FContainerNode(ContainerWidget));
@@ -162,7 +187,12 @@ namespace SlateIM
 
 	void BeginBorder(const FName BorderStyleName, EOrientation Orientation, bool bAbsorbMouse, FMargin ContentPadding)
 	{
-		BeginBorder(FAppStyle::GetBrush(BorderStyleName), Orientation, bAbsorbMouse, ContentPadding);
+		BeginBorder(FAppStyle::GetBrush(BorderStyleName), {.Orientation = Orientation, .bAbsorbMouse = bAbsorbMouse, .ContentPadding = ContentPadding});
+	}
+
+	void BeginBorder(const FName BorderStyleName, const FBorderParams& Params)
+	{
+		BeginBorder(FAppStyle::GetBrush(BorderStyleName), Params);
 	}
 
 	void EndBorder()
@@ -171,6 +201,11 @@ namespace SlateIM
 	}
 
 	bool BeginScrollBox(EOrientation Orientation)
+	{
+		return BeginScrollBox({.Orientation = Orientation});
+	}
+
+	bool BeginScrollBox(const FScrollBoxParams& Params)
 	{
 		TSharedPtr<SImScrollBox> ContainerWidget;
 
@@ -183,7 +218,7 @@ namespace SlateIM
 			{
 				ContainerWidget =
 					SNew(SImScrollBox)
-					.Orientation(Orientation)
+					.Orientation(Params.Orientation)
 					.OnUserScrolled_Lambda([ActivationData = Scope.GetOrCreateActivationMetadata().ToWeakPtr()](float Offset)
 					{
 						FSlateIMManager::Get().ActivateWidget(ActivationData.Pin());
@@ -207,12 +242,22 @@ namespace SlateIM
 		FSlateIMManager::Get().PopContainer<SImScrollBox>();
 	}
 
+	void BeginPopUp(const FName BorderStyleName, const FPopUpParams& Params)
+	{
+		BeginPopUp(FAppStyle::GetBrush(BorderStyleName), Params);
+	}
+
 	void BeginPopUp(const FName BorderStyleName, EOrientation Orientation, bool bAbsorbMouse, FMargin ContentPadding)
 	{
-		BeginPopUp(FAppStyle::GetBrush(BorderStyleName), Orientation, bAbsorbMouse, ContentPadding);
+		BeginPopUp(FAppStyle::GetBrush(BorderStyleName), {.Orientation = Orientation, .bAbsorbMouse = bAbsorbMouse, .ContentPadding = ContentPadding});
 	}
 
 	void BeginPopUp(const FSlateBrush* BorderBrush, EOrientation Orientation, bool bAbsorbMouse, FMargin ContentPadding)
+	{
+		BeginPopUp(BorderBrush, {.Orientation = Orientation, .bAbsorbMouse = bAbsorbMouse, .ContentPadding = ContentPadding});
+	}
+
+	void BeginPopUp(const FSlateBrush* BorderBrush, const FPopUpParams& Params)
 	{
 		TSharedPtr<SImPopUp> ContainerWidget;
 		{
@@ -236,7 +281,7 @@ namespace SlateIM
 		}
 
 		FSlateIMManager::Get().PushContainer(FContainerNode(ContainerWidget));
-		SlateIM::BeginBorder(BorderBrush, Orientation, bAbsorbMouse, ContentPadding);
+		SlateIM::BeginBorder(BorderBrush, Params.Orientation, Params.bAbsorbMouse, Params.ContentPadding);
 	}
 
 	void EndPopUp()
@@ -493,7 +538,12 @@ namespace SlateIM
 		BeginTableColumn(ColumnID);
 	}
 
-	void AddTableColumn(const FName& ColumnID, const FStringView& ColumnLabel)
+	void AddTableColumn(const FName& ColumnID, const FStringView& Label)
+	{
+		AddTableColumn(ColumnID, {.Label = Label});
+	}
+
+	void AddTableColumn(const FName& ColumnID, const FTableColumnParams& Params)
 	{
 		FSlateIMAlignmentState AlignmentState = FSlateIMManager::Get().SaveAlignmentState();
 
@@ -505,10 +555,10 @@ namespace SlateIM
 
 		NextTableColumn(ColumnID);
 
-		if (!ColumnLabel.IsEmpty())
+		if (!Params.Label.IsEmpty())
 		{
 			FSlateIMManager::Get().RestoreAlignmentState(AlignmentState);
-			SlateIM::Text(ColumnLabel);
+			SlateIM::Text(Params.Label);
 		}
 
 		EndTableColumn();
@@ -644,6 +694,10 @@ namespace SlateIM
 	}
 
 	bool BeginTableRowChildren(uint32 ParentRowId, bool bDefaultExpanded)
+	{
+		return BeginTableRowChildren({.ParentRowId = ParentRowId, .bDefaultExpanded = bDefaultExpanded});
+	}
+	bool BeginTableRowChildren(const FTableRowChildrenParams& Params)
 	{
 		// If there's an open cell, close it
 		if (TSharedPtr<SImStackBox> Cell = FSlateIMManager::Get().GetCurrentContainer<SImStackBox>())
@@ -657,37 +711,37 @@ namespace SlateIM
 			return false;
 		}
 
-		ParentRow->SetRowId(ParentRowId);
+		ParentRow->SetRowId(Params.ParentRowId);
 
 		SlateIM::BeginTableRow();
 
-		if (ParentRowId != 0)
+		if (Params.ParentRowId != 0)
 		{
 			if (TSharedPtr<SImTableView> OwningTable = ParentRow->GetOwningTable())
 			{
 				// If the default expanded state is true, we have a state because the default is unexpanded.
-				bool bHasExpansionState = bDefaultExpanded;
-				bool bExpanded = bDefaultExpanded;
+				bool bHasExpansionState = Params.bDefaultExpanded;
+				bool bExpanded = Params.bDefaultExpanded;
 				TSharedPtr<ISlateIMContainer> GrandParent = ParentRow->GetParent();
 
 				if (GrandParent->IsA<FSlateIMTableBody>())
 				{
 					TSharedPtr<FSlateIMTableBody> GrandParentRow = StaticCastSharedPtr<FSlateIMTableBody>(GrandParent);
 					
-					if (GrandParentRow->HasSavedExpansionState(ParentRowId))
+					if (GrandParentRow->HasSavedExpansionState(Params.ParentRowId))
 					{
 						bHasExpansionState = true;
-						bExpanded = GrandParentRow->GetSavedExpansionState(ParentRowId, bDefaultExpanded);
+						bExpanded = GrandParentRow->GetSavedExpansionState(Params.ParentRowId, Params.bDefaultExpanded);
 					}
 				}
 				else
 				{
 					TSharedPtr<FSlateIMTableRow> GrandParentBody = StaticCastSharedPtr<FSlateIMTableRow>(GrandParent);
 
-					if (GrandParentBody->HasSavedExpansionState(ParentRowId))
+					if (GrandParentBody->HasSavedExpansionState(Params.ParentRowId))
 					{
 						bHasExpansionState = true;
-						bExpanded = GrandParentBody->GetSavedExpansionState(ParentRowId, bDefaultExpanded);
+						bExpanded = GrandParentBody->GetSavedExpansionState(Params.ParentRowId, Params.bDefaultExpanded);
 					}
 				}
 
diff --git a/Engine/Plugins/Experimental/SlateIM/Source/SlateIM/Private/Implementation/SlateIM_Graphing.cpp b/Engine/Plugins/Experimental/SlateIM/Source/SlateIM/Private/Implementation/SlateIM_Graphing.cpp
--- a/Engine/Plugins/Experimental/SlateIM/Source/SlateIM/Private/Implementation/SlateIM_Graphing.cpp
+++ b/Engine/Plugins/Experimental/SlateIM/Source/SlateIM/Private/Implementation/SlateIM_Graphing.cpp
@@ -1,4 +1,4 @@
-﻿// Copyright Epic Games, Inc. All Rights Reserved.
+// Copyright Epic Games, Inc. All Rights Reserved.
 
 #include "SlateIM.h"
 
@@ -45,20 +45,30 @@ namespace SlateIM
 	}
 
 	void GraphLine(const TArrayView<FVector2D>& Points, const FLinearColor& LineColor, float LineThickness, const FDoubleRange& XViewRange, const FDoubleRange& YViewRange)
+	{
+		GraphLine(Points, {.XViewRange = XViewRange, .YViewRange = YViewRange, .LineColor = LineColor, .LineThickness = LineThickness});
+	}
+
+	void GraphLine(const TArrayView<FVector2D>& Points, const FGraphLinePointsParams& Params)
 	{
 		TSharedPtr<SImGraph> GraphWidget = GetCurrentGraphWidget();
 		if (ensureMsgf(GraphWidget, TEXT("Calling SlateIM::GraphLine() but the last widget is not a graph widget")))
 		{
-			GraphWidget->AddLineGraph(Points, LineColor, LineThickness, XViewRange, YViewRange);
+			GraphWidget->AddLineGraph(Points, Params);
 		}
 	}
 
 	void GraphLine(const TArrayView<double>& Values, const FLinearColor& LineColor, float LineThickness, const FDoubleRange& ViewRange)
+	{
+		GraphLine(Values, {.ViewRange = ViewRange, .LineColor = LineColor, .LineThickness = LineThickness});
+	}
+
+	void GraphLine(const TArrayView<double>& Values, const FGraphLineValuesParams& Params)
 	{
 		TSharedPtr<SImGraph> GraphWidget = GetCurrentGraphWidget();
 		if (ensureMsgf(GraphWidget, TEXT("Calling SlateIM::GraphLine() but the last widget is not a graph widget")))
 		{
-			GraphWidget->AddLineGraph(Values, LineColor, LineThickness, ViewRange);
+			GraphWidget->AddLineGraph(Values, Params);
 		}
 	}
 }
diff --git a/Engine/Plugins/Experimental/SlateIM/Source/SlateIM/Private/Implementation/SlateIM_Menu.cpp b/Engine/Plugins/Experimental/SlateIM/Source/SlateIM/Private/Implementation/SlateIM_Menu.cpp
--- a/Engine/Plugins/Experimental/SlateIM/Source/SlateIM/Private/Implementation/SlateIM_Menu.cpp
+++ b/Engine/Plugins/Experimental/SlateIM/Source/SlateIM/Private/Implementation/SlateIM_Menu.cpp
@@ -105,7 +105,7 @@ namespace SlateIM
 
 		SlateIM::Padding(0);
 
-		if (Button(InMenuName, &FSlateIMStyle::Get().GetWidgetStyle<FButtonStyle>("SlateIM.MenuButton")))
+		if (Button(InMenuName, {.Style = &FSlateIMStyle::Get().GetWidgetStyle<FButtonStyle>("SlateIM.MenuButton")}))
 		{
 			if (TSharedPtr<SImContextMenuAnchor> Menu = FSlateIMManager::Get().GetCurrentMenuRoot())
 			{
@@ -143,27 +143,42 @@ namespace SlateIM
 	}
 
 	bool AddMenuButton(const FStringView& RowText, const FStringView& ToolTipText)
+	{
+		return AddMenuButton(RowText, {.ToolTipText = ToolTipText});
+	}
+
+	bool AddMenuButton(const FStringView& RowText, const FMenuButtonParams& Params)
 	{
 		TSharedPtr<SImContextMenuAnchor> Anchor = FSlateIMManager::Get().GetCurrentMenuRoot();
 		checkf(Anchor, TEXT("Cannot add menu items without a current active menu anchor"));
 
-		return Anchor->AddMenuButton(RowText, ToolTipText);
+		return Anchor->AddMenuButton(RowText, Params.ToolTipText);
 	}
 
 	bool AddMenuToggleButton(const FStringView& RowText, bool& InOutCurrentState, const FStringView& ToolTipText)
+	{
+		return AddMenuToggleButton(RowText, InOutCurrentState, {.ToolTipText = ToolTipText});
+	}
+
+	bool AddMenuToggleButton(const FStringView& RowText, bool& InOutCurrentState, const FMenuButtonParams& Params)
 	{
 		TSharedPtr<SImContextMenuAnchor> Anchor = FSlateIMManager::Get().GetCurrentMenuRoot();
 		checkf(Anchor, TEXT("Cannot add menu items without a current active menu anchor"));
 
-		return Anchor->AddMenuToggleButton(RowText, InOutCurrentState, ToolTipText);
+		return Anchor->AddMenuToggleButton(RowText, InOutCurrentState, Params.ToolTipText);
 	}
 
 	bool AddMenuCheckButton(const FStringView& RowText, bool& InOutCurrentState, const FStringView& ToolTipText)
+	{
+		return AddMenuCheckButton(RowText, InOutCurrentState, {.ToolTipText = ToolTipText});
+	}
+
+	bool AddMenuCheckButton(const FStringView& RowText, bool& InOutCurrentState, const FMenuButtonParams& Params)
 	{
 		TSharedPtr<SImContextMenuAnchor> Anchor = FSlateIMManager::Get().GetCurrentMenuRoot();
 		checkf(Anchor, TEXT("Cannot add menu items without a current active menu anchor"));
 
-		return Anchor->AddMenuCheckButton(RowText, InOutCurrentState, ToolTipText);
+		return Anchor->AddMenuCheckButton(RowText, InOutCurrentState, Params.ToolTipText);
 	}
 
 	void BeginSubMenu(const FStringView& SubMenuText)
diff --git a/Engine/Plugins/Experimental/SlateIM/Source/SlateIM/Private/Implementation/SlateIM_Roots.cpp b/Engine/Plugins/Experimental/SlateIM/Source/SlateIM/Private/Implementation/SlateIM_Roots.cpp
--- a/Engine/Plugins/Experimental/SlateIM/Source/SlateIM/Private/Implementation/SlateIM_Roots.cpp
+++ b/Engine/Plugins/Experimental/SlateIM/Source/SlateIM/Private/Implementation/SlateIM_Roots.cpp
@@ -84,6 +84,11 @@ namespace SlateIM
 
 #if WITH_ENGINE
 	bool BeginViewportRoot(FName UniqueName, UGameViewportClient* ViewportClient, const FViewportRootLayout& Layout)
+	{
+		return BeginViewportRoot(UniqueName, ViewportClient, {.Layout = Layout});
+	}
+
+	bool BeginViewportRoot(FName UniqueName, UGameViewportClient* ViewportClient, const FViewportParams& Params)
 	{
 		if (!FSlateIMManager::Get().CanUpdateSlateIM())
 		{
@@ -113,7 +118,7 @@ namespace SlateIM
 			ViewportRoot = StaticCastSharedPtr<FSlateIMViewportRoot>(RootNode->RootWidget);
 		}
 
-		ViewportRoot->UpdateViewport(Layout);
+		ViewportRoot->UpdateViewport(Params.Layout);
 
 		FSlateIMManager::Get().BeginRoot(UniqueName);
 
@@ -121,6 +126,11 @@ namespace SlateIM
 	}
 
 	bool BeginViewportRoot(FName UniqueName, ULocalPlayer* LocalPlayer, const FViewportRootLayout& Layout)
+	{
+		return BeginViewportRoot(UniqueName, LocalPlayer, {.Layout = Layout});
+	}
+
+	bool BeginViewportRoot(FName UniqueName, ULocalPlayer* LocalPlayer, const FViewportParams& Params)
 	{
 		if (!FSlateIMManager::Get().CanUpdateSlateIM())
 		{
@@ -150,7 +160,7 @@ namespace SlateIM
 			ViewportRoot = StaticCastSharedPtr<FSlateIMViewportRoot>(RootNode->RootWidget);
 		}
 
-		ViewportRoot->UpdateViewport(Layout);
+		ViewportRoot->UpdateViewport(Params.Layout);
 
 		FSlateIMManager::Get().BeginRoot(UniqueName);
 
@@ -159,6 +169,11 @@ namespace SlateIM
 
 #if WITH_EDITOR
 	bool BeginViewportRoot(FName UniqueName, TSharedPtr<IAssetViewport> AssetViewport, const FViewportRootLayout& Layout)
+	{
+		return BeginViewportRoot(UniqueName, AssetViewport, {.Layout = Layout});
+	}
+
+	bool BeginViewportRoot(FName UniqueName, TSharedPtr<IAssetViewport> AssetViewport, const FViewportParams& Params)
 	{
 		if (!FSlateIMManager::Get().CanUpdateSlateIM())
 		{
@@ -188,7 +203,7 @@ namespace SlateIM
 			ViewportRoot = StaticCastSharedPtr<FSlateIMViewportRoot>(RootNode->RootWidget);
 		}
 
-		ViewportRoot->UpdateViewport(Layout);
+		ViewportRoot->UpdateViewport(Params.Layout);
 
 		FSlateIMManager::Get().BeginRoot(UniqueName);
 
diff --git a/Engine/Plugins/Experimental/SlateIM/Source/SlateIM/Private/Implementation/SlateIM_Tabs.cpp b/Engine/Plugins/Experimental/SlateIM/Source/SlateIM/Private/Implementation/SlateIM_Tabs.cpp
--- a/Engine/Plugins/Experimental/SlateIM/Source/SlateIM/Private/Implementation/SlateIM_Tabs.cpp
+++ b/Engine/Plugins/Experimental/SlateIM/Source/SlateIM/Private/Implementation/SlateIM_Tabs.cpp
@@ -87,6 +87,11 @@ namespace SlateIM
 	}
 
 	bool BeginTab(const FName& TabId, const FSlateIcon& TabIcon, const FText& TabTitle)
+	{
+		return BeginTab(TabId, {.TabIcon = TabIcon, .TabTitle = TabTitle});
+	}
+
+	bool BeginTab(const FName& TabId, const FTabParams& Params)
 	{
 		TSharedPtr<SImTabGroup> TabGroup = FSlateIMManager::Get().GetCurrentContainer<SImTabGroup>();
 		if (!ensureMsgf(TabGroup, TEXT("Cannot begin a Tab - the current container is not a Tab Group. Are your Begin and End function calls mismatched?")))
@@ -110,9 +115,9 @@ namespace SlateIM
 				ContainerWidget->SetOnTabDrawerClosed(FSimpleDelegate::CreateSP(ContainerWidget.Get(), &SImTab::OnDrawerClosed));
 				Scope.UpdateWidget(ContainerWidget);
 
-				// TabIcon and TabTitle are not updated after construction, this could be supported in the future but is not currently part of the Tab Manager API
-				ContainerWidget->TabIcon = TabIcon;
-				ContainerWidget->TabTitle = TabTitle;
+				// Params.TabIcon and Params.TabTitle are not updated after construction, this could be supported in the future but is not currently part of the Tab Manager API
+				ContainerWidget->TabIcon = Params.TabIcon;
+				ContainerWidget->TabTitle = Params.TabTitle;
 			}
 
 			if (!FSlateIMManager::Get().GetMutableCurrentRoot().Tabs.Contains(TabId))
diff --git a/Engine/Plugins/Experimental/SlateIM/Source/SlateIM/Private/Implementation/SlateIM_Util.cpp b/Engine/Plugins/Experimental/SlateIM/Source/SlateIM/Private/Implementation/SlateIM_Util.cpp
--- a/Engine/Plugins/Experimental/SlateIM/Source/SlateIM/Private/Implementation/SlateIM_Util.cpp
+++ b/Engine/Plugins/Experimental/SlateIM/Source/SlateIM/Private/Implementation/SlateIM_Util.cpp
@@ -1,4 +1,4 @@
-﻿// Copyright Epic Games, Inc. All Rights Reserved.
+// Copyright Epic Games, Inc. All Rights Reserved.
 
 #include "Misc/MessageDialog.h"
 #include "SlateIM.h"
@@ -28,6 +28,11 @@ namespace SlateIM
 	}
 
 	EAppReturnType::Type ModalDialog(EAppMsgType::Type MessageType, const FStringView& DialogText, EAppMsgCategory Category, const FStringView& DialogTitle)
+	{
+		return ModalDialog(MessageType, DialogText, {.Category = Category, .Title = DialogTitle});
+	}
+
+	EAppReturnType::Type ModalDialog(EAppMsgType::Type MessageType, const FStringView& DialogText, const FModalDialogParams& Params)
 	{
 		FSlateIMManager::Get().OnSlateIMModalOpened();
 
@@ -35,6 +40,6 @@ namespace SlateIM
 		{
 			FSlateIMManager::Get().OnSlateIMModalClosed();
 		};
-		return FMessageDialog::Open(Category, MessageType, FText::FromStringView(DialogText), FText::FromStringView(DialogTitle));
+		return FMessageDialog::Open(Params.Category, MessageType, FText::FromStringView(DialogText), FText::FromStringView(Params.Title));
 	}
 }
diff --git a/Engine/Plugins/Experimental/SlateIM/Source/SlateIM/Private/Implementation/SlateIM_Widgets.cpp b/Engine/Plugins/Experimental/SlateIM/Source/SlateIM/Private/Implementation/SlateIM_Widgets.cpp
--- a/Engine/Plugins/Experimental/SlateIM/Source/SlateIM/Private/Implementation/SlateIM_Widgets.cpp
+++ b/Engine/Plugins/Experimental/SlateIM/Source/SlateIM/Private/Implementation/SlateIM_Widgets.cpp
@@ -5,6 +5,7 @@
 
 #include "Misc/SlateIMLogging.h"
 #include "Misc/SlateIMManager.h"
+#include "Misc/SlateIMSlateResources.h"
 #include "Misc/SlateIMSlotData.h"
 #include "Misc/SlateIMWidgetScope.h"
 #include "Widgets/Input/SSearchableComboBox.h"
@@ -34,35 +35,46 @@ namespace SlateIM
 	}
 
 	void Text(const FStringView& InText, FSlateColor Color, const FTextBlockStyle* TextStyle /*= nullptr*/)
+	{
+		Text(InText, {.Color = Color, .Style = TextStyle});
+	}
+
+	void Text(const FStringView& InText, const FTextParams& Params)
 	{
 		FWidgetScope<STextBlock> Scope;
-		TSharedPtr<STextBlock> TextBlock = Scope.GetOrCreate([&InText, Color, TextStyle] ()
+		TSharedPtr<STextBlock> TextBlock = Scope.GetOrCreate([&InText, Params] ()
 			{
 				TSharedPtr<STextBlock> TextBlock =
 					SNew(STextBlock)
 					.AutoWrapText(true)
 					.Text(FText::FromStringView(InText))
-					.ColorAndOpacity(Color);
+					.ColorAndOpacity(Params.Color);
 
-				if (TextStyle)
+				if (Params.Style)
 				{
-					TextBlock->SetTextStyle(TextStyle);
+					TextBlock->SetTextStyle(Params.Style);
 				}
 				return TextBlock;
 			});
 
-		Scope.HashData(Color);
-		Scope.HashData(TextStyle);
+		Scope.HashData(Params.Color);
+		Scope.HashData(Params.Style);
 		Scope.HashStringView(InText);
-		if(Scope.IsDataHashDirty())
+
+		if (Scope.IsDataHashDirty())
 		{
 			TextBlock->SetText(FText::FromStringView(InText));
-			TextBlock->SetTextStyle(TextStyle);
-			TextBlock->SetColorAndOpacity(Color);
+			TextBlock->SetTextStyle(Params.Style);
+			TextBlock->SetColorAndOpacity(Params.Color);
 		}
 	}
 
 	bool EditableText(FString& InOutText, const FStringView& HintText, const FEditableTextBoxStyle* TextStyle)
+	{
+		return EditableText(InOutText, {.HintText = HintText, .Style = TextStyle});
+	}
+
+	bool EditableText(FString& InOutText, const FEditableTextParams& Params)
 	{
 		FWidgetScope<SEditableTextBox> Scope;
 		TSharedPtr<SEditableTextBox> EditableText = Scope.GetWidget();
@@ -70,9 +82,9 @@ namespace SlateIM
 
 		bool bWasActivated = false;
 
-		Scope.HashData(TextStyle);
+		Scope.HashData(Params.Style);
 		Scope.HashStringView(InOutText);
-		Scope.HashStringView(HintText);
+		Scope.HashStringView(Params.HintText);
 
 		if (!EditableText)
 		{
@@ -82,7 +94,7 @@ namespace SlateIM
 				SNew(SEditableTextBox)
 				.MinDesiredWidth(MinWidth)
 				.Text(FText::FromStringView(InOutText))
-				.HintText(FText::FromStringView(HintText))
+				.HintText(FText::FromStringView(Params.HintText))
 				.OnTextChanged_Lambda([ActivationData](const FText& NewText)
 				{
 					FSlateIMManager::Get().ActivateWidget(ActivationData.Pin());
@@ -92,9 +104,9 @@ namespace SlateIM
 					FSlateIMManager::Get().ActivateWidget(ActivationData.Pin());
 				});
 
-			if (TextStyle)
+			if (Params.Style)
 			{
-				EditableText->SetStyle(TextStyle);
+				EditableText->SetStyle(Params.Style);
 			}
 			
 			Scope.UpdateWidget(EditableText);
@@ -107,8 +119,8 @@ namespace SlateIM
 
 			if (bIsHashDirty)
 			{
-				EditableText->SetStyle(TextStyle);
-				EditableText->SetHintText(FText::FromStringView(HintText));
+				EditableText->SetStyle(Params.Style);
+				EditableText->SetHintText(FText::FromStringView(Params.HintText));
 			}
 			
 			if (bWasActivated)
@@ -126,7 +138,7 @@ namespace SlateIM
 		return bWasActivated;
 	}
 
-	void Image_Internal(const FSlateColor& ColorAndOpacity, FVector2D DesiredSize, TFunctionRef<const FSlateBrush*(const TSharedRef<SImage>&)> GetBrush)
+	void Image_Internal(const FImageParams& Params, TFunctionRef<const FSlateBrush*(const TSharedRef<SImage>&)> GetBrush)
 	{
 		FWidgetScope<SImage> Scope(Defaults::Padding, HAlign_Left, VAlign_Top);
 		TSharedPtr<SImage> ImageWidget = Scope.GetWidget();
@@ -135,8 +147,8 @@ namespace SlateIM
 		{
 			ImageWidget =
 				SNew(SImage)
-				.ColorAndOpacity(ColorAndOpacity)
-				.DesiredSizeOverride(DesiredSize == FVector2D::ZeroVector ? TOptional<FVector2D>() : DesiredSize);
+				.ColorAndOpacity(Params.ColorAndOpacity)
+				.DesiredSizeOverride(Params.DesiredSize == FVector2D::ZeroVector ? TOptional<FVector2D>() : Params.DesiredSize);
 
 			ImageWidget->SetImage(GetBrush(ImageWidget.ToSharedRef()));
 			Scope.UpdateWidget(ImageWidget);
@@ -144,29 +156,49 @@ namespace SlateIM
 		else
 		{
 			ImageWidget->SetImage(GetBrush(ImageWidget.ToSharedRef()));
-			ImageWidget->SetColorAndOpacity(ColorAndOpacity);
-			ImageWidget->SetDesiredSizeOverride(DesiredSize == FVector2D::ZeroVector ? TOptional<FVector2D>() : DesiredSize);
+			ImageWidget->SetColorAndOpacity(Params.ColorAndOpacity);
+			ImageWidget->SetDesiredSizeOverride(Params.DesiredSize == FVector2D::ZeroVector ? TOptional<FVector2D>() : Params.DesiredSize);
 		}
 	}
 
 	void Image(const FSlateBrush* ImageBrush, const FSlateColor& ColorAndOpacity, FVector2D DesiredSize)
 	{
-		Image_Internal(ColorAndOpacity, DesiredSize, [ImageBrush](const TSharedRef<SImage>&) { return ImageBrush; });
+		Image_Internal({.ColorAndOpacity = ColorAndOpacity, .DesiredSize = DesiredSize}, [ImageBrush](const TSharedRef<SImage>&) { return ImageBrush; });
+	}
+
+	void Image(const FSlateBrush* ImageBrush, const FImageParams& Params)
+	{
+		Image_Internal(Params, [ImageBrush](const TSharedRef<SImage>&) { return ImageBrush; });
 	}
 
 	void Image(const FName ImageStyleName, const FSlateColor& ColorAndOpacity, FVector2D DesiredSize)
 	{
-		Image(FAppStyle::Get().GetBrush(ImageStyleName), ColorAndOpacity, DesiredSize);
+		Image(FAppStyle::Get().GetBrush(ImageStyleName), {.ColorAndOpacity = ColorAndOpacity, .DesiredSize = DesiredSize});
+	}
+
+	void Image(const FName ImageStyleName, const FImageParams& Params)
+	{
+		Image(FAppStyle::Get().GetBrush(ImageStyleName), Params);
 	}
 
 	void Image(const FSlateColor& ColorAndOpacity, FVector2D DesiredSize)
+	{
+		Image({.ColorAndOpacity = ColorAndOpacity, .DesiredSize = DesiredSize});
+	}
+
+	void Image(const FImageParams& Params)
 	{
 		static FSlateBrush DefaultBrush;
-		Image(&DefaultBrush, ColorAndOpacity, DesiredSize);
+		Image(&DefaultBrush, Params);
 	}
 
 #if WITH_ENGINE
 	void Image(UTexture2D* ImageTexture, const FSlateColor& ColorAndOpacity, FVector2D DesiredSize)
+	{
+		Image(ImageTexture, {.ColorAndOpacity = ColorAndOpacity, .DesiredSize = DesiredSize});
+	}
+
+	void Image(UTexture2D* ImageTexture, const FImageParams& Params)
 	{
 		auto MakeTextureBrush = [ImageTexture]() -> FSlateBrush
 		{
@@ -180,7 +212,7 @@ namespace SlateIM
 			return FSlateNoResource();
 		};
 		
-		Image_Internal(ColorAndOpacity, DesiredSize, [&MakeTextureBrush, ImageTexture](const TSharedRef<SImage>& ImageWidget)
+		Image_Internal(Params, [&MakeTextureBrush, ImageTexture](const TSharedRef<SImage>& ImageWidget)
 		{
 			TSharedPtr<FUObjectImageResource> Resource = ImageWidget->GetMetaData<FUObjectImageResource>();
 			if (!Resource)
@@ -198,6 +230,11 @@ namespace SlateIM
 	}
 
 	void Image(UTextureRenderTarget2D* ImageRenderTarget, const FSlateColor& ColorAndOpacity, FVector2D DesiredSize)
+	{
+		Image(ImageRenderTarget, {.ColorAndOpacity = ColorAndOpacity, .DesiredSize = DesiredSize});
+	}
+
+	void Image(UTextureRenderTarget2D* ImageRenderTarget, const FImageParams& Params)
 	{
 		auto MakeRTBrush = [ImageRenderTarget]() -> FSlateBrush
 		{
@@ -211,7 +248,7 @@ namespace SlateIM
 			return FSlateNoResource();
 		};
 		
-		Image_Internal(ColorAndOpacity, DesiredSize, [&MakeRTBrush, ImageRenderTarget](const TSharedRef<SImage>& ImageWidget)
+		Image_Internal(Params, [&MakeRTBrush, ImageRenderTarget](const TSharedRef<SImage>& ImageWidget)
 		{
 			TSharedPtr<FUObjectImageResource> Resource = ImageWidget->GetMetaData<FUObjectImageResource>();
 			if (!Resource)
@@ -230,7 +267,12 @@ namespace SlateIM
 
 	void Image(UMaterialInterface* ImageMaterial, FVector2D BrushSize, const FSlateColor& ColorAndOpacity, FVector2D DesiredSize)
 	{
-		auto MakeMaterialBrush = [ImageMaterial, &BrushSize]() -> FSlateBrush
+		Image(ImageMaterial, BrushSize, {.ColorAndOpacity = ColorAndOpacity, .DesiredSize = DesiredSize});
+	}
+
+	void Image(UMaterialInterface* ImageMaterial, FVector2D BrushSize, const FImageParams& Params)
+	{
+		auto MakeMaterialBrush = [ImageMaterial, BrushSize]() -> FSlateBrush
 		{
 			if (ImageMaterial)
 			{
@@ -239,7 +281,7 @@ namespace SlateIM
 			return FSlateNoResource();
 		};
 		
-		Image_Internal(ColorAndOpacity, DesiredSize, [&MakeMaterialBrush, ImageMaterial](const TSharedRef<SImage>& ImageWidget)
+		Image_Internal(Params, [&MakeMaterialBrush, ImageMaterial](const TSharedRef<SImage>& ImageWidget)
 		{
 			TSharedPtr<FUObjectImageResource> Resource = ImageWidget->GetMetaData<FUObjectImageResource>();
 			if (!Resource)
@@ -259,11 +301,26 @@ namespace SlateIM
 
 	bool Button(const FStringView& InText, const FButtonStyle* InStyle)
 	{
+		return Button(InText, {.Style = InStyle});
+	}
+
+	bool Button(const FStringView& InText, const bool bEnabled, const FButtonStyle* InStyle)
+	{
+		return Button(InText, {.bEnabled = bEnabled, .Style = InStyle});
+	}
+
+	bool Button(const FStringView& InText, const FButtonParams& Params)
+	{
+		if (!Params.bEnabled)
+		{
+			BeginDisabledState();
+		}
+
 		FWidgetScope<SImButton> Scope;
 		TSharedPtr<SImButton> ButtonWidget = Scope.GetWidget();
 
 		Scope.HashStringView(InText);
-		Scope.HashData(InStyle);
+		Scope.HashData(Params.Style);
 
 		bool bWasClicked = false;
 		if (!ButtonWidget)
@@ -278,9 +335,9 @@ namespace SlateIM
 
 			ButtonWidget->SetText(InText);
 
-			if (InStyle)
+			if (Params.Style)
 			{
-				ButtonWidget->SetButtonStyle(InStyle);
+				ButtonWidget->SetButtonStyle(Params.Style);
 			}
 
 			Scope.UpdateWidget(ButtonWidget);
@@ -292,38 +349,32 @@ namespace SlateIM
 			if (Scope.IsDataHashDirty())
 			{
 				ButtonWidget->SetText(InText);
-				if (InStyle)
+				if (Params.Style)
 				{
-					ButtonWidget->SetButtonStyle(InStyle);
+					ButtonWidget->SetButtonStyle(Params.Style);
 				}
 			}
 		}
 
-		return bWasClicked;
-	}
-
-	bool Button(const FStringView& InText, const bool bEnabled, const FButtonStyle* InStyle)
-	{
-		if (!bEnabled)
+		if (!Params.bEnabled)
 		{
-			BeginDisabledState();
+			EndDisabledState();
 		}
 
-		const bool bResult = Button(InText, InStyle);
+		return bWasClicked;
+	}
 
-		if (!bEnabled)
-		{
-			EndDisabledState();
-		}
 
-		return bResult;
+	bool CheckBox(const FStringView& InLabel, bool& InOutCurrentState, const FCheckBoxStyle* CheckBoxStyle)
+	{
+		return CheckBox(InOutCurrentState, {.Label = InLabel, .CheckBoxStyle = CheckBoxStyle});
 	}
 
-	bool CheckBox(const FStringView& InText, bool& InOutCurrentState, const FCheckBoxStyle* CheckBoxStyle)
+	bool CheckBox(bool& InOutCurrentState, const FCheckBoxParams& Params)
 	{
 		ECheckBoxState CurrentEnumState = InOutCurrentState ? ECheckBoxState::Checked : ECheckBoxState::Unchecked;
 
-		const bool bValueChanged = CheckBox(InText, CurrentEnumState, CheckBoxStyle);
+		const bool bValueChanged = CheckBox(CurrentEnumState, Params);
 		if (bValueChanged)
 		{
 			InOutCurrentState = CurrentEnumState == ECheckBoxState::Checked;
@@ -332,15 +383,20 @@ namespace SlateIM
 		return bValueChanged;
 	}
 
-	bool CheckBox(const FStringView& InText, ECheckBoxState& InOutCurrentState, const FCheckBoxStyle* CheckBoxStyle)
+	bool CheckBox(const FStringView& InLabel, ECheckBoxState& InOutCurrentState, const FCheckBoxStyle* CheckBoxStyle)
+	{
+		return CheckBox(InOutCurrentState, {.Label = InLabel, .CheckBoxStyle = CheckBoxStyle});
+	}
+
+	bool CheckBox(ECheckBoxState& InOutCurrentState, const FCheckBoxParams& Params)
 	{
 		FWidgetScope<SImCheckBox> Scope;
 		TSharedPtr<SImCheckBox> CheckboxWidget = Scope.GetWidget();
 
 		bool bValueChanged = false;
 
-		Scope.HashStringView(InText);
-		Scope.HashData(CheckBoxStyle);
+		Scope.HashStringView(Params.Label);
+		Scope.HashData(Params.CheckBoxStyle);
 
 		if (!CheckboxWidget)
 		{
@@ -352,11 +408,11 @@ namespace SlateIM
 					FSlateIMManager::Get().ActivateWidget(ActivationData.Pin());
 				});
 
-			CheckboxWidget->SetText(FText::FromStringView(InText));
+			CheckboxWidget->SetText(FText::FromStringView(Params.Label));
 
-			if (CheckBoxStyle)
+			if (Params.CheckBoxStyle)
 			{
-				CheckboxWidget->SetStyle(CheckBoxStyle);
+				CheckboxWidget->SetStyle(Params.CheckBoxStyle);
 			}
 
 			Scope.UpdateWidget(CheckboxWidget);
@@ -375,10 +431,10 @@ namespace SlateIM
 
 			if (Scope.IsDataHashDirty())
 			{
-				CheckboxWidget->SetText(FText::FromStringView(InText));
-				if (CheckBoxStyle)
+				CheckboxWidget->SetText(FText::FromStringView(Params.Label));
+				if (Params.CheckBoxStyle)
 				{
-					CheckboxWidget->SetStyle(CheckBoxStyle);
+					CheckboxWidget->SetStyle(Params.CheckBoxStyle);
 				}
 			}
 		}
@@ -387,12 +443,12 @@ namespace SlateIM
 	}
 
 	template<typename NumericType>
-	bool SpinBox_Internal(NumericType& InOutValue, TOptional<NumericType> Min, TOptional<NumericType> Max, const FSpinBoxStyle* SpinBoxStyle = nullptr)
+	bool SpinBox_Internal(NumericType& InOutValue, const FSpinBoxParams<NumericType>& Params)
 	{
 		FWidgetScope<SSpinBox<NumericType>> Scope(Defaults::Padding, Defaults::HAlign, Defaults::VAlign, Defaults::bAutoSize, Defaults::InputWidgetWidth);
 		TSharedPtr<SSpinBox<NumericType>> SpinBoxWidget = Scope.GetWidget();
 
-		Scope.HashData(SpinBoxStyle);
+		Scope.HashData(Params.Style);
 
 		class FSpinBoxState : public ISlateMetaData
 		{
@@ -408,8 +464,8 @@ namespace SlateIM
 			TWeakPtr<FSlateIMWidgetActivationMetadata> ActivationData = Scope.GetOrCreateActivationMetadata().ToWeakPtr();
 			SpinBoxWidget =
 				SNew(SSpinBox<NumericType>)
-				.MinValue(Min)
-				.MaxValue(Max)
+				.MinValue(Params.Min)
+				.MaxValue(Params.Max)
 				.Value(InOutValue)
 				.OnBeginSliderMovement_Lambda([SpinBoxState](){ SpinBoxState->bIsChanging = true; })
 				.OnEndSliderMovement_Lambda([SpinBoxState, ActivationData](NumericType NewVal)
@@ -429,15 +485,15 @@ namespace SlateIM
 			SpinBoxWidget->AddMetadata(SpinBoxState);
 			Scope.UpdateWidget(SpinBoxWidget);
 
-			if (SpinBoxStyle)
+			if (Params.Style)
 			{
-				SpinBoxWidget->SetWidgetStyle(SpinBoxStyle);
+				SpinBoxWidget->SetWidgetStyle(Params.Style);
 			}
 		}
 		else
 		{
-			SpinBoxWidget->SetMinValue(Min);
-			SpinBoxWidget->SetMaxValue(Max);
+			SpinBoxWidget->SetMinValue(Params.Min);
+			SpinBoxWidget->SetMaxValue(Params.Max);
 
 			bValueChanged = Scope.IsActivatedThisFrame();
 			if (bValueChanged)
@@ -453,9 +509,9 @@ namespace SlateIM
 				}
 			}
 
-			if (SpinBoxStyle && Scope.IsDataHashDirty())
+			if (Params.Style && Scope.IsDataHashDirty())
 			{
-				SpinBoxWidget->SetWidgetStyle(SpinBoxStyle);
+				SpinBoxWidget->SetWidgetStyle(Params.Style);
 				SpinBoxWidget->InvalidateStyle();
 			}
 		}
@@ -466,51 +522,74 @@ namespace SlateIM
 
 	bool SpinBox(float& InOutValue, TOptional<float> Min, TOptional<float> Max, const FSpinBoxStyle* SpinBoxStyle)
 	{
-		return SpinBox_Internal<float>(InOutValue, Min, Max, SpinBoxStyle);
+		return SpinBox_Internal<float>(InOutValue, {.Min = Min, .Max = Max, .Style = SpinBoxStyle});
 	}
 
 	bool SpinBox(double& InOutValue, TOptional<double> Min, TOptional<double> Max, const FSpinBoxStyle* SpinBoxStyle)
 	{
-		return SpinBox_Internal<double>(InOutValue, Min, Max, SpinBoxStyle);
+		return SpinBox_Internal<double>(InOutValue, {.Min = Min, .Max = Max, .Style = SpinBoxStyle});
 	}
 
 	bool SpinBox(int32& InOutValue, TOptional<int32> Min, TOptional<int32> Max, const FSpinBoxStyle* SpinBoxStyle)
 	{
-		return SpinBox_Internal<int32>(InOutValue, Min, Max, SpinBoxStyle);
+		return SpinBox_Internal<int32>(InOutValue, {.Min = Min, .Max = Max, .Style = SpinBoxStyle});
+	}
+
+	template<>
+	bool SpinBox<float>(float& InOutValue, const FSpinBoxParams<float>& Params)
+	{
+		return SpinBox_Internal<float>(InOutValue, Params);
+	}
+
+	template<>
+	bool SpinBox<double>(double& InOutValue, const FSpinBoxParams<double>& Params)
+	{
+		return SpinBox_Internal<double>(InOutValue, Params);
+	}
+
+	template<>
+	bool SpinBox<int32>(int32& InOutValue, const FSpinBoxParams<int32>& Params)
+	{
+		return SpinBox_Internal<int32>(InOutValue, Params);
 	}
 
 	bool Slider(float& InOutValue, float Min, float Max, float Step, const FSliderStyle* SliderStyle /*= nullptr*/)
+	{
+		return Slider(InOutValue, {.Min = Min, .Max = Max, .Step = Step, .Style = SliderStyle});
+	}
+
+	bool Slider(float& InOutValue, const FSliderParams& Params)
 	{
 		FWidgetScope<SSlider> Scope(Defaults::Padding, Defaults::HAlign, Defaults::VAlign, Defaults::bAutoSize, Defaults::InputWidgetWidth);
 		TSharedPtr<SSlider> SliderWidget = Scope.GetWidget();
 
-		Scope.HashData(SliderStyle);
+		Scope.HashData(Params.Style);
 
 		bool bValueChanged = false;
 		if (!SliderWidget)
 		{
 			SliderWidget =
 				SNew(SSlider)
-				.MinValue(Min)
-				.MaxValue(Max)
-				.StepSize(Step)
+				.MinValue(Params.Min)
+				.MaxValue(Params.Max)
+				.StepSize(Params.Step)
 				.Value(InOutValue)
 				.OnValueChanged_Lambda([ActivationData = Scope.GetOrCreateActivationMetadata().ToWeakPtr()](float NewVal)
 				{
 					FSlateIMManager::Get().ActivateWidget(ActivationData.Pin());
 				});
 
-			if (SliderStyle)
+			if (Params.Style)
 			{
-				SliderWidget->SetStyle(SliderStyle);
+				SliderWidget->SetStyle(Params.Style);
 			}
 
 			Scope.UpdateWidget(SliderWidget);
 		}
 		else
 		{
-			SliderWidget->SetMinAndMaxValues(Min, Max);
-			SliderWidget->SetStepSize(Step);
+			SliderWidget->SetMinAndMaxValues(Params.Min, Params.Max);
+			SliderWidget->SetStepSize(Params.Step);
 
 			bValueChanged = Scope.IsActivatedThisFrame();
 			if (bValueChanged)
@@ -522,9 +601,9 @@ namespace SlateIM
 				SliderWidget->SetValue(InOutValue);
 			}
 			
-			if (Scope.IsDataHashDirty() && SliderStyle)
+			if (Scope.IsDataHashDirty() && Params.Style)
 			{
-				SliderWidget->SetStyle(SliderStyle);
+				SliderWidget->SetStyle(Params.Style);
 			}
 		}
 
@@ -532,11 +611,16 @@ namespace SlateIM
 	}
 
 	void ProgressBar(TOptional<float> Percent, const FProgressBarStyle* ProgressBarStyle /*= nullptr*/)
+	{
+		ProgressBar(Percent, {.Style = ProgressBarStyle});
+	}
+
+	void ProgressBar(TOptional<float> Percent, const FProgressBarParams& Params)
 	{
 		FWidgetScope<SProgressBar> Scope(Defaults::Padding, Defaults::HAlign, Defaults::VAlign, Defaults::bAutoSize, Defaults::InputWidgetWidth);
 		TSharedPtr<SProgressBar> ProgressBarWidget = Scope.GetWidget();
 
-		Scope.HashData(ProgressBarStyle);
+		Scope.HashData(Params.Style);
 
 		if (!ProgressBarWidget)
 		{
@@ -544,9 +628,9 @@ namespace SlateIM
 				SNew(SProgressBar)
 				.Percent(Percent);
 
-			if (ProgressBarStyle)
+			if (Params.Style)
 			{
-				ProgressBarWidget->SetStyle(ProgressBarStyle);
+				ProgressBarWidget->SetStyle(Params.Style);
 			}
 
 			Scope.UpdateWidget(ProgressBarWidget);
@@ -555,14 +639,14 @@ namespace SlateIM
 		{
 			ProgressBarWidget->SetPercent(Percent);
 			
-			if (Scope.IsDataHashDirty() && ProgressBarStyle)
+			if (Scope.IsDataHashDirty() && Params.Style)
 			{
-				ProgressBarWidget->SetStyle(ProgressBarStyle);
+				ProgressBarWidget->SetStyle(Params.Style);
 			}
 		}
 	}
 
-	bool SearchableComboBox(const TArray<FString>& ComboItems, int32& InOutSelectedItemIndex, bool bForceRefresh, const FComboBoxStyle* InComboStyle)
+	bool SearchableComboBox_Impl(const TArray<FString>& ComboItems, int32& InOutSelectedItemIndex, const FComboBoxParams& Params)
 	{
 		FWidgetScope<SSearchableComboBox> Scope;
 		TSharedPtr<SSearchableComboBox> ComboWidget = Scope.GetWidget();
@@ -591,7 +675,7 @@ namespace SlateIM
 			SAssignNew(ComboWidget, SSearchableComboBox)
 			.OptionsSource(&WidgetIMData->Data)
 			.InitiallySelectedItem(WidgetIMData->Data.IsValidIndex(InOutSelectedItemIndex) ? WidgetIMData->Data[InOutSelectedItemIndex] : nullptr)
-			.ComboBoxStyle(InComboStyle ? InComboStyle : &FAppStyle::Get().GetWidgetStyle<FComboBoxStyle>("ComboBox"))
+			.ComboBoxStyle(Params.Style ? Params.Style : &FAppStyle::Get().GetWidgetStyle<FComboBoxStyle>("ComboBox"))
 			.Content()[ComboText]
 			.OnSelectionChanged_Lambda([ActivationData = Scope.GetOrCreateActivationMetadata().ToWeakPtr()](TSharedPtr<FString> NewVal, ESelectInfo::Type SelectType)
 			{
@@ -628,7 +712,7 @@ namespace SlateIM
 		else
 		{
 			TSharedPtr<FComboBoxData> ComboBoxData = ComboWidget->GetMetaData<FComboBoxData>();
-			if (bForceRefresh)
+			if (Params.bForceRefresh)
 			{
 				ComboBoxData->Data = BuildComboOptions();
 				if (InOutSelectedItemIndex != INDEX_NONE && !ComboWidget->IsOpen() && ComboBoxData->Data.IsValidIndex(InOutSelectedItemIndex))
@@ -652,12 +736,17 @@ namespace SlateIM
 		return bValueChanged;
 	}
 
-	bool ComboBox(const TArray<FString>& ComboItems, int32& InOutSelectedItemIndex, bool bForceRefresh, const FComboBoxStyle* InComboStyle)
+	bool SearchableComboBox(const TArray<FString>& ComboItems, int32& InOutSelectedItemIndex, bool bForceRefresh, const FComboBoxStyle* InComboStyle)
+	{
+		return SearchableComboBox_Impl(ComboItems, InOutSelectedItemIndex, {.bForceRefresh = bForceRefresh, .Style = InComboStyle});
+	}
+
+	bool ComboBox_Impl(const TArray<FString>& ComboItems, int32& InOutSelectedItemIndex, const FComboBoxParams& Params)
 	{
 		FWidgetScope<STextComboBox> Scope;
 		TSharedPtr<STextComboBox> ComboWidget = Scope.GetWidget();
 
-		Scope.HashData(InComboStyle);
+		Scope.HashData(Params.Style);
 		Scope.HashData(InOutSelectedItemIndex);
 		
 		using FComboBoxData = FSlateIMDataStore<TArray<TSharedPtr<FString>>>;
@@ -688,9 +777,9 @@ namespace SlateIM
 					FSlateIMManager::Get().ActivateWidget(ActivationData.Pin());
 				});
 
-			if (InComboStyle)
+			if (Params.Style)
 			{
-				ComboWidget->SetStyle(InComboStyle);
+				ComboWidget->SetStyle(Params.Style);
 			}
 			
 			ComboWidget->AddMetadata(WidgetIMData);
@@ -700,7 +789,7 @@ namespace SlateIM
 		else
 		{
 			TSharedPtr<FComboBoxData> ComboBoxData = ComboWidget->GetMetaData<FComboBoxData>();
-			if (bForceRefresh)
+			if (Params.bForceRefresh)
 			{
 				ComboBoxData->Data = BuildComboOptions();
 				if (InOutSelectedItemIndex != INDEX_NONE && !ComboWidget->IsOpen() && ComboBoxData->Data.IsValidIndex(InOutSelectedItemIndex))
@@ -729,14 +818,34 @@ namespace SlateIM
 					ComboWidget->SetSelectedItem(ComboBoxData->Data[InOutSelectedItemIndex]);
 				}
 
-				ComboWidget->SetStyle(InComboStyle);
+				ComboWidget->SetStyle(Params.Style);
 			}
 		}
 
 		return bValueChanged;
 	}
 
+	bool ComboBox(const TArray<FString>& ComboItems, int32& InOutSelectedItemIndex, bool bForceRefresh, const FComboBoxStyle* InComboStyle)
+	{
+		return ComboBox_Impl(ComboItems, InOutSelectedItemIndex, {.bForceRefresh = bForceRefresh, .Style = InComboStyle});
+	}
+
+	bool ComboBox(const TArray<FString>& ComboItems, int32& InOutSelectedItemIndex, const FComboBoxParams& Params)
+	{
+		if (Params.bSearchable)
+		{
+			return SearchableComboBox_Impl(ComboItems, InOutSelectedItemIndex, Params);
+		}
+
+		return ComboBox_Impl(ComboItems, InOutSelectedItemIndex, Params);
+	}
+
 	bool SelectionList(const TArray<FString>& ListItems, int32& InOutSelectedItemIndex, bool bForceRefresh, const FTableViewStyle* InStyle)
+	{
+		return SelectionList(ListItems, InOutSelectedItemIndex, {.bForceRefresh = bForceRefresh, .Style = InStyle});
+	}
+
+	bool SelectionList(const TArray<FString>& ListItems, int32& InOutSelectedItemIndex, const FSelectionListParams& Params)
 	{
 		SCOPED_NAMED_EVENT_TEXT("SlateIM::SelectionList", FColorList::Goldenrod);
 		using ListViewType = SListView<TSharedPtr<FString>>;
@@ -745,7 +854,7 @@ namespace SlateIM
 		FWidgetScope<ListViewType> Scope;
 		TSharedPtr<ListViewType> ListWidget = Scope.GetWidget();
 
-		Scope.HashData(InStyle);
+		Scope.HashData(Params.Style);
 
 		auto BuildListItems = [&]()
 		{
@@ -764,7 +873,7 @@ namespace SlateIM
 			TSharedRef<FListViewData> ListViewData = MakeShared<FListViewData>(BuildListItems());
 			ListWidget = SNew(ListViewType)
 				.SelectionMode(ESelectionMode::Single)
-				.ListViewStyle(InStyle)
+				.ListViewStyle(Params.Style)
 				.ListItemsSource(&(ListViewData->Data))
 				.OnSelectionChanged_Lambda([ActivationData = Scope.GetOrCreateActivationMetadata().ToWeakPtr()](TSharedPtr<FString> NewVal, ESelectInfo::Type SelectType)
 				{
@@ -791,7 +900,7 @@ namespace SlateIM
 		else
 		{
 			TSharedPtr<FListViewData> ListViewData = ListWidget->GetMetaData<FListViewData>();
-			if (bForceRefresh)
+			if (Params.bForceRefresh)
 			{
 				ListViewData->Data = BuildListItems();
 				ListWidget->RequestListRefresh();
@@ -819,7 +928,7 @@ namespace SlateIM
 
 			if (Scope.IsDataHashDirty())
 			{
-				ListWidget->SetStyle(InStyle);
+				ListWidget->SetStyle(Params.Style);
 			}
 		}
 
diff --git a/Engine/Plugins/Experimental/SlateIM/Source/SlateIM/Private/Misc/SlateIMExamples.cpp b/Engine/Plugins/Experimental/SlateIM/Source/SlateIM/Private/Misc/SlateIMExamples.cpp
--- a/Engine/Plugins/Experimental/SlateIM/Source/SlateIM/Private/Misc/SlateIMExamples.cpp
+++ b/Engine/Plugins/Experimental/SlateIM/Source/SlateIM/Private/Misc/SlateIMExamples.cpp
@@ -18,6 +18,7 @@
 #if WITH_EDITOR
 #include "SLevelViewport.h"
 #endif
+
 #if WITH_ENGINE
 #include "Engine/Canvas.h"
 #include "Engine/Engine.h"
@@ -109,7 +110,7 @@ void FSlateIMTestWidget::Draw()
 	SlateIM::AddMenuButton(TEXT("Option 1"));
 	SlateIM::EndMenuBar();
 
-	SlateIM::BeginBorder(FAppStyle::GetBrush("ToolPanel.GroupBorder"), Orient_Vertical, bAbsorbMouse);
+	SlateIM::BeginBorder(FAppStyle::GetBrush("ToolPanel.GroupBorder"), {.bAbsorbMouse = bAbsorbMouse});
 	// Basic perf measurement, outside of the scrollbox so that it's "pinned"
 	if (TimeSinceLastUpdate > 0.5f)
 	{
@@ -216,7 +217,7 @@ void FSlateIMTestWidget::DrawBasics()
 				// EditableText Example
 				{
 					SlateIM::BeginHorizontalStack();
-					SlateIM::EditableText(ComboItemToAdd, TEXT("Add Combo Item"));
+					SlateIM::EditableText(ComboItemToAdd, {.HintText = TEXT("Add Combo Item")});
 					if (SlateIM::IsFocused(SlateIM::EFocusDepth::IncludingDescendants))
 					{
 						SlateIM::BeginPopUp();
@@ -247,22 +248,22 @@ void FSlateIMTestWidget::DrawBasics()
 			}
 
 			// CheckBox examples
-			SlateIM::CheckBox(TEXT("Check Box"), CheckState);
+			SlateIM::CheckBox(CheckState, {.Label = TEXT("Check Box")});
 			if (CheckState)
 			{
 				SlateIM::Text(TEXT("Basic Text"));
-				SlateIM::Text(TEXT("Text With Color"), FLinearColor::Green);
-				SlateIM::Text(TEXT("Text With style color"), FStyleColors::Primary);
+				SlateIM::Text(TEXT("Text With Color"), {.Color = FLinearColor::Green});
+				SlateIM::Text(TEXT("Text With style color"), {.Color = FStyleColors::Primary});
 			}
 	
-			if (SlateIM::CheckBox(TEXT("Undetermined Check Box"), CheckStateEnum))
+			if (SlateIM::CheckBox(CheckStateEnum, {.Label = TEXT("Undetermined Check Box")}))
 			{
 				UE_LOG(LogSlateIM, Log, TEXT("Check Box State Changed"));
 			}
 
 			// ComboBox example
 			SlateIM::BeginHorizontalStack();
-			if (SlateIM::ComboBox(ComboBoxItems, SelectedItemIndex, bRefreshComboItems))
+			if (SlateIM::ComboBox(ComboBoxItems, SelectedItemIndex, {.bForceRefresh = bRefreshComboItems}))
 			{
 				if (ComboBoxItems.IsValidIndex(SelectedItemIndex))
 				{
@@ -273,7 +274,7 @@ void FSlateIMTestWidget::DrawBasics()
 			SlateIM::EndHorizontalStack();
 
 			SlateIM::BeginHorizontalStack();
-			if (SlateIM::SearchableComboBox(ComboBoxItems, SearchableComboBoxSelectedItemIndex, bRefreshComboItems))
+			if (SlateIM::ComboBox(ComboBoxItems, SearchableComboBoxSelectedItemIndex, {.bForceRefresh = bRefreshComboItems, .bSearchable = true}))
 			{
 				if (ComboBoxItems.IsValidIndex(SearchableComboBoxSelectedItemIndex))
 				{
@@ -313,7 +314,7 @@ void FSlateIMTestWidget::DrawBasics()
 			SlateIM::EndVerticalStack();
 
 			// SelectionList example
-			SlateIM::SelectionList(ComboBoxItems, SelectedItem, bRefreshComboItems);
+			SlateIM::SelectionList(ComboBoxItems, SelectedItem, {.bForceRefresh = bRefreshComboItems});
 
 			SlateIM::EndHorizontalStack();
 
@@ -342,7 +343,7 @@ void FSlateIMTestWidget::DrawBasics()
 			SlateIM::BeginVerticalStack();
 			{
 				// Slider example
-				if (SlateIM::Slider(SliderVal, 0, SliderMax, 1))
+				if (SlateIM::Slider(SliderVal, {.Max = SliderMax, .Step = 1.f}))
 				{
 					UE_LOG(LogSlateIM, Log, TEXT("Slider Value Changed [%f]"), SliderVal);
 				}
@@ -353,10 +354,10 @@ void FSlateIMTestWidget::DrawBasics()
 				// SpinBox Examples
 				{
 					SlateIM::Padding(FMargin(0, 10, 0, 5));
-					SlateIM::SpinBox(SliderVal, 0.0f, SliderMax);
+					SlateIM::SpinBox(SliderVal, {.Min = 0.0f, .Max = SliderMax});
 
 					SlateIM::Padding(FMargin(0, 10, 0, 5));
-					SlateIM::SpinBox(IntValue, 0, IntMax);
+					SlateIM::SpinBox(IntValue, {.Min = 0, .Max = IntMax});
 				}
 			}
 			SlateIM::EndVerticalStack();
@@ -401,7 +402,7 @@ void FSlateIMTestWidget::DrawBasics()
 		{
 			SCOPED_NAMED_EVENT_TEXT("DisabledState example", FColorList::Goldenrod);
 			// DisabledState example
-			SlateIM::CheckBox(TEXT("Disable Everything Below Me"), bShouldBeDisabled);
+			SlateIM::CheckBox(bShouldBeDisabled, {.Label = TEXT("Disable Everything Below Me")});
 			if (bShouldBeDisabled)
 			{
 				SlateIM::BeginDisabledState();
@@ -420,36 +421,36 @@ void FSlateIMTestWidget::DrawBasics()
 
 			// This part is only shown if the menu is open
 			SlateIM::AddMenuSection(TEXT("Menu Section 1"));
-			if (SlateIM::AddMenuButton(TEXT("Menu Item 1"), TEXT("Menu Item Tool Tip 1")))
+			if (SlateIM::AddMenuButton(TEXT("Menu Item 1"), {.ToolTipText = TEXT("Menu Item Tool Tip 1")}))
 			{
 				UE_LOG(LogSlateIM, Log, TEXT("Menu Item One menu option clicked"));
 			}
 
-			SlateIM::AddMenuButton(TEXT("Menu Item 2"), TEXT("Menu Item Tool Tip 2"));
-			SlateIM::AddMenuButton(TEXT("Menu Item 3"), TEXT("Menu Item Tool Tip 3"));
-			SlateIM::AddMenuButton(TEXT("Menu Item 4"), TEXT("Menu Item Tool Tip 4"));
+			SlateIM::AddMenuButton(TEXT("Menu Item 2"), {.ToolTipText = TEXT("Menu Item Tool Tip 2")});
+			SlateIM::AddMenuButton(TEXT("Menu Item 3"), {.ToolTipText = TEXT("Menu Item Tool Tip 3")});
+			SlateIM::AddMenuButton(TEXT("Menu Item 4"), {.ToolTipText = TEXT("Menu Item Tool Tip 4")});
 
-			SlateIM::AddMenuCheckButton(TEXT("Menu Item With Check"), MenuCheckState, TEXT("Click to toggle check"));
+			SlateIM::AddMenuCheckButton(TEXT("Menu Item With Check"), MenuCheckState, {.ToolTipText = TEXT("Click to toggle check")});
 
-			if (SlateIM::AddMenuToggleButton(TEXT("Menu Item With Toggle"), MenuToggleState, TEXT("Toggle this box")))
+			if (SlateIM::AddMenuToggleButton(TEXT("Menu Item With Toggle"), MenuToggleState, {.ToolTipText = TEXT("Toggle this box")}))
 			{
 				UE_LOG(LogSlateIM, Log, TEXT("Menu Item With Toggle clicked"));
 			}
 
 			SlateIM::AddMenuSeparator();
 			SlateIM::BeginSubMenu(TEXT("Sub Menu"));
-			SlateIM::AddMenuButton(TEXT("SubMenu Item 1"), TEXT("Menu Item Tool Tip 1"));
-			SlateIM::AddMenuButton(TEXT("SubMenu Item 2"), TEXT("Menu Item Tool Tip 2"));
-			SlateIM::AddMenuButton(TEXT("SubMenu Item 3"), TEXT("Menu Item Tool Tip 3"));
-			SlateIM::AddMenuButton(TEXT("SubMenu Item 4"), TEXT("Menu Item Tool Tip 4"));
+			SlateIM::AddMenuButton(TEXT("SubMenu Item 1"), {.ToolTipText = TEXT("Menu Item Tool Tip 1")});
+			SlateIM::AddMenuButton(TEXT("SubMenu Item 2"), {.ToolTipText = TEXT("Menu Item Tool Tip 2")});
+			SlateIM::AddMenuButton(TEXT("SubMenu Item 3"), {.ToolTipText = TEXT("Menu Item Tool Tip 3")});
+			SlateIM::AddMenuButton(TEXT("SubMenu Item 4"), {.ToolTipText = TEXT("Menu Item Tool Tip 4")});
 			SlateIM::EndSubMenu();
-			SlateIM::AddMenuButton(TEXT("Menu Item 5"), TEXT("A Menu Item between 2 SubMenus"));
-			SlateIM::AddMenuButton(TEXT("Menu Item 6"), TEXT("A Menu Item between 2 SubMenus"));
+			SlateIM::AddMenuButton(TEXT("Menu Item 5"), {.ToolTipText = TEXT("A Menu Item between 2 SubMenus")});
+			SlateIM::AddMenuButton(TEXT("Menu Item 6"), {.ToolTipText = TEXT("A Menu Item between 2 SubMenus")});
 			SlateIM::BeginSubMenu(TEXT("Sub Menu 2"));
-			SlateIM::AddMenuButton(TEXT("SubMenu2 Item 1"), TEXT("Menu Item Tool Tip 1"));
-			SlateIM::AddMenuButton(TEXT("SubMenu2 Item 2"), TEXT("Menu Item Tool Tip 2"));
-			SlateIM::AddMenuButton(TEXT("SubMenu2 Item 3"), TEXT("Menu Item Tool Tip 3"));
-			SlateIM::AddMenuButton(TEXT("SubMenu2 Item 4"), TEXT("Menu Item Tool Tip 4"));
+			SlateIM::AddMenuButton(TEXT("SubMenu2 Item 1"), {.ToolTipText = TEXT("Menu Item Tool Tip 1")});
+			SlateIM::AddMenuButton(TEXT("SubMenu2 Item 2"), {.ToolTipText = TEXT("Menu Item Tool Tip 2")});
+			SlateIM::AddMenuButton(TEXT("SubMenu2 Item 3"), {.ToolTipText = TEXT("Menu Item Tool Tip 3")});
+			SlateIM::AddMenuButton(TEXT("SubMenu2 Item 4"), {.ToolTipText = TEXT("Menu Item Tool Tip 4")});
 			SlateIM::EndSubMenu();
 			SlateIM::EndContextMenuAnchor();
 		}
@@ -579,7 +580,7 @@ void FSlateIMTestWidget::DrawTables()
 			}
 
 			NumTableItems = NumItems;
-			SlateIM::CheckBox(TEXT("Live Update Table?"), bShouldLiveUpdateTable);
+			SlateIM::CheckBox(bShouldLiveUpdateTable, {.Label = TEXT("Live Update Table?")});
 			if (bShouldLiveUpdateTable)
 			{
 				NumTableItems = LiveNumItems;
@@ -604,7 +605,7 @@ void FSlateIMTestWidget::DrawTables()
 				{
 					SlateIM::Padding(FMargin(5, 0));
 					SlateIM::VAlign(VAlign_Center);// Centers the button in the row
-					SlateIM::Text(FString::Printf(TEXT("Item %d/%d"), i + 1, NumItems), FColor::MakeRedToGreenColorFromScalar(static_cast<float>(i) / NumItems));
+					SlateIM::Text(FString::Printf(TEXT("Item %d/%d"), i + 1, NumItems), {.Color = FColor::MakeRedToGreenColorFromScalar(static_cast<float>(i) / NumItems)});
 					SlateIM::Padding(FMargin(0));
 
 					if (SlateIM::Button(TEXT("Click")))
@@ -640,7 +641,7 @@ void FSlateIMTestWidget::DrawTables()
 			SlateIM::InitialTableColumnWidth(80.f);
 			SlateIM::Maximize();
 			SlateIM::HAlign(HAlign_Center);
-			SlateIM::AddTableColumn(TEXT("Button"), TEXT("Button"));
+			SlateIM::AddTableColumn(TEXT("Button"), {.Label = TEXT("Button")});
 			SlateIM::EndTableHeader();
 			SlateIM::BeginTableBody();
 			for (int32 i = 0; i < NumTableItems; ++i)
@@ -655,11 +656,11 @@ void FSlateIMTestWidget::DrawTables()
 
 					if (!bIsSelected)
 					{
-						SlateIM::Text(FString::Printf(TEXT("Item %d/%d"), i + 1, NumTableItems), FColor::MakeRedToGreenColorFromScalar(static_cast<float>(i) / NumTableItems));
+						SlateIM::Text(FString::Printf(TEXT("Item %d/%d"), i + 1, NumTableItems), {.Color = FColor::MakeRedToGreenColorFromScalar(static_cast<float>(i) / NumTableItems)});
 					}
 					else
 					{
-						SlateIM::Text(FString::Printf(TEXT("Item %d/%d [Selected]"), i + 1, NumTableItems), FColor::MakeRedToGreenColorFromScalar(static_cast<float>(i) / NumTableItems));
+						SlateIM::Text(FString::Printf(TEXT("Item %d/%d [Selected]"), i + 1, NumTableItems), {.Color = FColor::MakeRedToGreenColorFromScalar(static_cast<float>(i) / NumTableItems)});
 					}
 				}
 
@@ -680,7 +681,7 @@ void FSlateIMTestWidget::DrawTables()
 		
 		// Tree Example
 		{
-			SlateIM::CheckBox(TEXT("Show Canada?"), bShouldAddCanada);
+			SlateIM::CheckBox(bShouldAddCanada, {.Label = TEXT("Show Canada?")});
 
 			SCOPED_NAMED_EVENT_TEXT("Tree Example", FColorList::Goldenrod);
 			const FTableRowStyle* TableRowStyle = &FAppStyle::Get().GetWidgetStyle<FTableRowStyle>("TableView.AlternatingRow");
@@ -689,9 +690,9 @@ void FSlateIMTestWidget::DrawTables()
 			SlateIM::MaxHeight(200.f);
 			SlateIM::HAlign(HAlign_Fill);
 			SlateIM::VAlign(VAlign_Fill);
-			SlateIM::BeginTable(nullptr, TableRowStyle);
-			SlateIM::AddTableColumn(TEXT("Name"), TEXT("Name"));
-			SlateIM::AddTableColumn(TEXT("Type"), TEXT("Type"));
+			SlateIM::BeginTable({.RowStyle = TableRowStyle});
+			SlateIM::AddTableColumn(TEXT("Name"), {.Label = TEXT("Name")});
+			SlateIM::AddTableColumn(TEXT("Type"), {.Label = TEXT("Type")});
 			{
 				if (SlateIM::NextTableCell())
 				{
@@ -731,7 +732,7 @@ void FSlateIMTestWidget::DrawTables()
 							SlateIM::Text(TEXT("Country"));
 						}
 						static constexpr FLazyName CanadaChildren = TEXT("CanadaChildren");
-						if (SlateIM::BeginTableRowChildren(GetTypeHash(CanadaChildren)))
+						if (SlateIM::BeginTableRowChildren({.ParentRowId = GetTypeHash(CanadaChildren)}))
 						{
 							if (SlateIM::NextTableCell())
 							{
@@ -743,7 +744,7 @@ void FSlateIMTestWidget::DrawTables()
 								SlateIM::VAlign(VAlign_Center);
 								SlateIM::Text(TEXT("Province"));
 							}
-							if (SlateIM::BeginTableRowChildren(ChildRowId++))
+							if (SlateIM::BeginTableRowChildren({.ParentRowId = static_cast<uint32>(ChildRowId++)}))
 							{
 								if (SlateIM::NextTableCell())
 								{
@@ -768,7 +769,7 @@ void FSlateIMTestWidget::DrawTables()
 								SlateIM::VAlign(VAlign_Center);
 								SlateIM::Text(TEXT("Province"));
 							}
-							if (SlateIM::BeginTableRowChildren(ChildRowId++))
+							if (SlateIM::BeginTableRowChildren({.ParentRowId = static_cast<uint32>(ChildRowId++)}))
 							{
 								if (SlateIM::NextTableCell())
 								{
@@ -796,8 +797,13 @@ void FSlateIMTestWidget::DrawTables()
 						SlateIM::VAlign(VAlign_Center);
 						SlateIM::Text(TEXT("Country"));
 					}
+#if WITH_ENGINE
 					constexpr uint32 RowID = UE::HashStringFNV1a32("UnitedStatesChildren");
-					if (SlateIM::BeginTableRowChildren(RowID))
+#else
+					constexpr uint32 RowID = 0;
+#endif
+
+					if (SlateIM::BeginTableRowChildren({.ParentRowId = RowID}))
 					{
 						if (SlateIM::NextTableCell())
 						{
@@ -859,12 +865,12 @@ void FSlateIMTestWidget::DrawTables()
 
 		// Dynamically added/removed child row
 		{
-			SlateIM::CheckBox(TEXT("Show Child?"), bShouldAddChildRow);
+			SlateIM::CheckBox(bShouldAddChildRow, {.Label = TEXT("Show Child?")});
 
 			SlateIM::Fill();
 			SlateIM::HAlign(HAlign_Fill);
 			SlateIM::BeginTable();
-			SlateIM::AddTableColumn(TEXT("Column"), TEXT("Column"));
+			SlateIM::AddTableColumn(TEXT("Column"), {.Label = TEXT("Column")});
 
 			if (SlateIM::NextTableCell())
 			{
@@ -907,7 +913,7 @@ void FSlateIMTestWidget::DrawTables()
 			SlateIM::Text(TEXT("Columns:"));
 			for (int32 i = 0; i < Columns.Num(); ++i)
 			{
-				SlateIM::CheckBox(GetColumnName(i), Columns[i]);
+				SlateIM::CheckBox(Columns[i], {.Label = GetColumnName(i)});
 			}
 			SlateIM::EndHorizontalStack();
 
@@ -915,7 +921,7 @@ void FSlateIMTestWidget::DrawTables()
 			SlateIM::Text(TEXT("Rows:"));
 			for (int32 i = 0; i < Rows.Num(); ++i)
 			{
-				SlateIM::CheckBox(GetRowName(i), Rows[i]);
+				SlateIM::CheckBox(Rows[i], {.Label = GetRowName(i)});
 			}
 			SlateIM::EndHorizontalStack();
 
@@ -929,7 +935,7 @@ void FSlateIMTestWidget::DrawTables()
 			{
 				if (Columns[i])
 				{
-					SlateIM::AddTableColumn(FName(GetColumnName(i)), GetColumnName(i));
+					SlateIM::AddTableColumn(FName(GetColumnName(i)), {.Label = GetColumnName(i)});
 				}
 			}
 			for (int32 i = 0; i < Rows.Num(); ++i)
@@ -972,7 +978,7 @@ void FSlateIMTestWidget::DrawGraphs()
 		SlateIM::Fill();
 		SlateIM::MinHeight(200.f);
 		SlateIM::BeginGraph();
-		SlateIM::GraphLine(SquareGraphValues.Compact(), FLinearColor::White, 3.f, FDoubleRange(0, 1.0));
+		SlateIM::GraphLine(SquareGraphValues.Compact(), {.ViewRange = FDoubleRange(0, 1.0), .LineColor = FLinearColor::White, .LineThickness = 3.f});
 		SlateIM::EndGraph();
 
 		SlateIM::VAlign(VAlign_Center);
@@ -1004,9 +1010,9 @@ void FSlateIMTestWidget::DrawGraphs()
 		SlateIM::Fill();
 		SlateIM::MinHeight(200.f);
 		SlateIM::BeginGraph();
-		SlateIM::GraphLine(SinGraphValues.Compact(), FColor::Orange, 1.f, FDoubleRange(NextSinX - 100.0, NextSinX), FDoubleRange(-1.5, 1.5));
-		SlateIM::GraphLine(CosGraphValues.Compact(), FLinearColor::Green, 1.f, FDoubleRange(NextCosX - 100.0, NextCosX), FDoubleRange(-1.5, 1.5));
-		SlateIM::GraphLine(TanGraphValues.Compact(), FColor::Magenta, 1.f, FDoubleRange(NextTanX - 100.0, NextTanX), FDoubleRange(-1.5, 1.5));
+		SlateIM::GraphLine(SinGraphValues.Compact(), {.XViewRange = FDoubleRange(NextSinX - 100.0, NextSinX), .YViewRange = FDoubleRange(-1.5, 1.5), .LineColor = FColor::Orange});
+		SlateIM::GraphLine(CosGraphValues.Compact(), {.XViewRange = FDoubleRange(NextCosX - 100.0, NextCosX), .YViewRange = FDoubleRange(-1.5, 1.5), .LineColor = FLinearColor::Green});
+		SlateIM::GraphLine(TanGraphValues.Compact(), {.XViewRange = FDoubleRange(NextTanX - 100.0, NextTanX), .YViewRange = FDoubleRange(-1.5, 1.5), .LineColor = FColor::Magenta});
 		SlateIM::EndGraph();
 
 		SlateIM::MinWidth(50.f);
@@ -1014,13 +1020,13 @@ void FSlateIMTestWidget::DrawGraphs()
 		SlateIM::BeginVerticalStack();
 		SlateIM::Fill();
 		SlateIM::VAlign(VAlign_Center);
-		SlateIM::Text(FString::Printf(TEXT("%0.3lf"), SinGraphValues.Last().Y), FColor::Orange);
+		SlateIM::Text(FString::Printf(TEXT("%0.3lf"), SinGraphValues.Last().Y), {.Color = FColor::Orange});
 		SlateIM::Fill();
 		SlateIM::VAlign(VAlign_Center);
-		SlateIM::Text(FString::Printf(TEXT("%0.3lf"), CosGraphValues.Last().Y), FLinearColor::Green);
+		SlateIM::Text(FString::Printf(TEXT("%0.3lf"), CosGraphValues.Last().Y), {.Color = FLinearColor::Green});
 		SlateIM::Fill();
 		SlateIM::VAlign(VAlign_Center);
-		SlateIM::Text(FString::Printf(TEXT("%0.3lf"), TanGraphValues.Last().Y), FColor::Magenta);
+		SlateIM::Text(FString::Printf(TEXT("%0.3lf"), TanGraphValues.Last().Y), {.Color = FColor::Magenta});
 		SlateIM::EndVerticalStack();
 	}
 	SlateIM::EndHorizontalStack();
@@ -1153,7 +1159,7 @@ void FSlateIMTestWidget::DrawInputs()
 			SlateIM::BeginBorder(&WBrush);
 			SlateIM::HAlign(HAlign_Center);
 			SlateIM::VAlign(VAlign_Center);
-			SlateIM::Text(TEXT("W"), FLinearColor::White);
+			SlateIM::Text(TEXT("W"), {.Color = FLinearColor::White});
 			SlateIM::EndBorder();
 		}
 		SlateIM::Spacer(FVector2D(50, 50));
@@ -1168,7 +1174,7 @@ void FSlateIMTestWidget::DrawInputs()
 			SlateIM::BeginBorder(&ABrush);
 			SlateIM::HAlign(HAlign_Center);
 			SlateIM::VAlign(VAlign_Center);
-			SlateIM::Text(TEXT("A"), FLinearColor::White);
+			SlateIM::Text(TEXT("A"), {.Color = FLinearColor::White});
 			SlateIM::EndBorder();
 		}
 		{
@@ -1180,7 +1186,7 @@ void FSlateIMTestWidget::DrawInputs()
 			SlateIM::BeginBorder(&SBrush);
 			SlateIM::HAlign(HAlign_Center);
 			SlateIM::VAlign(VAlign_Center);
-			SlateIM::Text(TEXT("S"), FLinearColor::White);
+			SlateIM::Text(TEXT("S"), {.Color = FLinearColor::White});
 			SlateIM::EndBorder();
 		}
 		{
@@ -1192,7 +1198,7 @@ void FSlateIMTestWidget::DrawInputs()
 			SlateIM::BeginBorder(&DBrush);
 			SlateIM::HAlign(HAlign_Center);
 			SlateIM::VAlign(VAlign_Center);
-			SlateIM::Text(TEXT("D"), FLinearColor::White);
+			SlateIM::Text(TEXT("D"), {.Color = FLinearColor::White});
 			SlateIM::EndBorder();
 		}
 		SlateIM::EndHorizontalStack();
@@ -1224,7 +1230,7 @@ void FSlateIMTestWidget::DrawInputs()
 			SlateIM::BeginBorder(&LMBBrush);
 			SlateIM::HAlign(HAlign_Center);
 			SlateIM::VAlign(VAlign_Center);
-			SlateIM::Text(TEXT("LMB"), FLinearColor::White);
+			SlateIM::Text(TEXT("LMB"), {.Color = FLinearColor::White});
 			SlateIM::EndBorder();
 		}
 		{
@@ -1236,7 +1242,7 @@ void FSlateIMTestWidget::DrawInputs()
 			SlateIM::BeginBorder(&RMBBrush);
 			SlateIM::HAlign(HAlign_Center);
 			SlateIM::VAlign(VAlign_Center);
-			SlateIM::Text(TEXT("RMB"), FLinearColor::White);
+			SlateIM::Text(TEXT("RMB"), {.Color = FLinearColor::White});
 			SlateIM::EndBorder();
 		}
 		SlateIM::EndHorizontalStack();
@@ -1277,14 +1283,14 @@ void FSlateIMTestWidget::DrawTabs()
 				SlateIM::TabSplitterSizeCoefficient(0.3f);
 				SlateIM::BeginTabStack();
 				{
-					if (SlateIM::BeginTab(TEXT("Left Tab"), FSlateIcon(FAppStyle::GetAppStyleSetName(), TEXT("Symbols.LeftArrow"))))
+					if (SlateIM::BeginTab(TEXT("Left Tab"), {.TabIcon = FSlateIcon(FAppStyle::GetAppStyleSetName(), TEXT("Symbols.LeftArrow"))}))
 					{
 						// Fill so that the text autowrapping updates properly as the splitter is adjusted by the user
 						SlateIM::HAlign(HAlign_Fill);
 						SlateIM::Text(TEXT("Left Tab taking up 30% of the horizontal space"));
 					}
 					SlateIM::EndTab();
-					if (SlateIM::BeginTab(TEXT("Right Tab"), FSlateIcon(FAppStyle::GetAppStyleSetName(), TEXT("Symbols.RightArrow"))))
+					if (SlateIM::BeginTab(TEXT("Right Tab"), {.TabIcon = FSlateIcon(FAppStyle::GetAppStyleSetName(), TEXT("Symbols.RightArrow"))}))
 					{
 						// Fill so that the text autowrapping updates properly as the splitter is adjusted by the user
 						SlateIM::HAlign(HAlign_Fill);
@@ -1353,7 +1359,7 @@ void FSlateIMTestWidget::DrawTabs()
 			{
 				for (int32 i = 0; i < DynamicTabCount; i++)
 				{
-					if (SlateIM::BeginTab(*FString::Printf(TEXT("DynamicTab%d"), i), FSlateIcon(), FText::Format(INVTEXT("Tab {0}"), FText::AsNumber(i + 1))))
+					if (SlateIM::BeginTab(*FString::Printf(TEXT("DynamicTab%d"), i), {.TabTitle = FText::Format(INVTEXT("Tab {0}"), FText::AsNumber(i + 1))}))
 					{
 						SlateIM::Text(FString::Printf(TEXT("Tab %d Content"), i + 1));
 					}
@@ -1443,7 +1449,7 @@ void FSlateIMTestWidget::DrawFilter()
 
 		// Draw the filter input and filter the list
 		bool bFruitListChanged = false;
-		if (SlateIM::TextFilter(FruitFilter, TEXT("Search fruits...")))
+		if (SlateIM::TextFilter(FruitFilter, {.HintText = TEXT("Search fruits...")}))
 		{
 			FilteredFruitList = FruitFilter.Filter(FruitList);
 			bFruitListChanged = true;
@@ -1453,7 +1459,7 @@ void FSlateIMTestWidget::DrawFilter()
 		SlateIM::Fill();
 		SlateIM::HAlign(HAlign_Fill);
 		int32 SelectedFruitIndex = INDEX_NONE;
-		SlateIM::SelectionList(FilteredFruitList.GetValue(), SelectedFruitIndex, bFruitListChanged);
+		SlateIM::SelectionList(FilteredFruitList.GetValue(), SelectedFruitIndex, {.bForceRefresh = bFruitListChanged});
 	}
 	SlateIM::EndVerticalStack();
 }
@@ -1465,16 +1471,16 @@ void FSlateStyleBrowser::DrawWindow(float DeltaTime)
 	SlateIM::VAlign(VAlign_Fill);
 	SlateIM::Fill();
 	SlateIM::BeginVerticalStack();
-	SlateIM::EditableText(SearchString, TEXT("Search Styles..."));
+	SlateIM::EditableText(SearchString, {.HintText = TEXT("Search Styles...")});
 	
 	SlateIM::HAlign(HAlign_Fill);
 	SlateIM::VAlign(VAlign_Fill);
 	SlateIM::Fill();
 	SlateIM::BeginTable();
-	SlateIM::AddTableColumn(TEXT("Name"), TEXT("Name"));
-	SlateIM::AddTableColumn(TEXT("Preview"), TEXT("Preview"));
+	SlateIM::AddTableColumn(TEXT("Name"), {.Label = TEXT("Name")});
+	SlateIM::AddTableColumn(TEXT("Preview"), {.Label = TEXT("Preview")});
 	SlateIM::InitialTableColumnWidth(80.f);
-	SlateIM::AddTableColumn(TEXT("Button"), TEXT("Button"));
+	SlateIM::AddTableColumn(TEXT("Button"), {.Label = TEXT("Button")});
 
 	const auto& Style = static_cast<const SlateIM::Private::FExposedSlateStyle&>(FAppStyle::Get());
 	const FString LowerCaseSearchString = SearchString.ToLower();
@@ -1555,7 +1561,7 @@ void FSlateStyleBrowser::DrawWindow(float DeltaTime)
 					SlateIM::HAlign(HAlign_Center);
 					SlateIM::VAlign(VAlign_Center);
 					SlateIM::Fill();
-					SlateIM::Text(TEXT("The quick brown fox jumps over the lazy dog."), WidgetStyle.Value);
+					SlateIM::Text(TEXT("The quick brown fox jumps over the lazy dog."), {.Style = WidgetStyle.Value});
 				}
 			
 				DrawCopyCell(WidgetStyle.Key);
@@ -1583,7 +1589,7 @@ void FSlateStyleBrowser::DrawWindow(float DeltaTime)
 					SlateIM::HAlign(HAlign_Center);
 					SlateIM::VAlign(VAlign_Center);
 					SlateIM::Fill();
-					SlateIM::EditableText(PreviewText, TEXT("Hint text..."), WidgetStyle.Value);
+					SlateIM::EditableText(PreviewText, {.HintText = TEXT("Hint text..."), .Style = WidgetStyle.Value});
 				}
 				
 				DrawCopyCell(WidgetStyle.Key);
@@ -1611,7 +1617,7 @@ void FSlateStyleBrowser::DrawWindow(float DeltaTime)
 					SlateIM::HAlign(HAlign_Center);
 					SlateIM::VAlign(VAlign_Center);
 					SlateIM::Fill();
-					SlateIM::Button(TEXT("Click Me"), WidgetStyle.Value);
+					SlateIM::Button(TEXT("Click Me"), {.Style = WidgetStyle.Value});
 				}
 				
 				DrawCopyCell(WidgetStyle.Key);
@@ -1639,7 +1645,7 @@ void FSlateStyleBrowser::DrawWindow(float DeltaTime)
 					SlateIM::HAlign(HAlign_Center);
 					SlateIM::VAlign(VAlign_Center);
 					SlateIM::Fill();
-					SlateIM::SpinBox(SpinBoxValue, -100.f, 100.f, WidgetStyle.Value);
+					SlateIM::SpinBox(SpinBoxValue, {.Min = -100.f, .Max = 100.f, .Style = WidgetStyle.Value});
 				}
 				
 				DrawCopyCell(WidgetStyle.Key);
@@ -1667,7 +1673,7 @@ void FSlateStyleBrowser::DrawWindow(float DeltaTime)
 					SlateIM::HAlign(HAlign_Center);
 					SlateIM::VAlign(VAlign_Center);
 					SlateIM::Fill();
-					SlateIM::Slider(SliderValue, 0.f, 100.f, 0.1f, WidgetStyle.Value);
+					SlateIM::Slider(SliderValue, {.Max = 100.f, .Step = 0.1f, .Style = WidgetStyle.Value});
 				}
 				
 				DrawCopyCell(WidgetStyle.Key);
@@ -1695,7 +1701,7 @@ void FSlateStyleBrowser::DrawWindow(float DeltaTime)
 					SlateIM::HAlign(HAlign_Center);
 					SlateIM::VAlign(VAlign_Center);
 					SlateIM::Fill();
-					SlateIM::ProgressBar(SliderValue / 100.f, WidgetStyle.Value);
+					SlateIM::ProgressBar(SliderValue / 100.f, {.Style = WidgetStyle.Value});
 				}
 				
 				DrawCopyCell(WidgetStyle.Key);
@@ -1723,8 +1729,7 @@ void FSlateStyleBrowser::DrawWindow(float DeltaTime)
 					SlateIM::HAlign(HAlign_Center);
 					SlateIM::VAlign(VAlign_Center);
 					SlateIM::Fill();
-					constexpr bool bForceRefresh = false;
-					SlateIM::ComboBox(Options, SelectedComboIndex, bForceRefresh, WidgetStyle.Value);
+					SlateIM::ComboBox(Options, SelectedComboIndex, {.Style = WidgetStyle.Value});
 				}
 				
 				DrawCopyCell(WidgetStyle.Key);
@@ -1757,14 +1762,14 @@ void FSlateStyleBrowser::DrawWindow(float DeltaTime)
 						SlateIM::HAlign(HAlign_Center);
 						SlateIM::VAlign(VAlign_Top);
 						SlateIM::Fill();
-						SlateIM::SelectionList(Options, SelectedListIndex, false, WidgetStyle.Value);
+						SlateIM::SelectionList(Options, SelectedListIndex, {.Style = WidgetStyle.Value});
 
 						SlateIM::HAlign(HAlign_Center);
 						SlateIM::VAlign(VAlign_Top);
 						SlateIM::Fill();
-						SlateIM::BeginTable(WidgetStyle.Value);
-						SlateIM::AddTableColumn(TEXT("Column 1"), TEXT("Column 1"));
-						SlateIM::AddTableColumn(TEXT("Column 2"), TEXT("Column 2"));
+						SlateIM::BeginTable({.Style = WidgetStyle.Value});
+						SlateIM::AddTableColumn(TEXT("Column 1"), {.Label = TEXT("Column 1")});
+						SlateIM::AddTableColumn(TEXT("Column 2"), {.Label = TEXT("Column 2")});
 						SlateIM::NextTableCell();
 						SlateIM::Text(TEXT("Cell 1"));
 						SlateIM::NextTableCell();
@@ -1811,7 +1816,7 @@ void FSlateIMTestViewportWidget::DrawWidget(float DeltaTime)
 {
 	if (GEngine && GEngine->GameViewport)
 	{
-		if (SlateIM::BeginViewportRoot("SlateIMTestSuiteViewport", GEngine->GameViewport, Layout))
+		if (SlateIM::BeginViewportRoot("SlateIMTestSuiteViewport", GEngine->GameViewport, {.Layout = Layout}))
 		{
 			TestWidget.Draw();
 		}
@@ -1823,7 +1828,7 @@ void FSlateIMTestViewportWidget::DrawWidget(float DeltaTime)
 		TSharedPtr<SLevelViewport> LevelViewport = StaticCastSharedPtr<SLevelViewport>(GCurrentLevelEditingViewportClient->GetEditorViewportWidget());
 		if (LevelViewport.IsValid())
 		{
-			if (SlateIM::BeginViewportRoot("SlateIMTestSuiteViewport", LevelViewport, Layout))
+			if (SlateIM::BeginViewportRoot("SlateIMTestSuiteViewport", LevelViewport, {.Layout = Layout}))
 			{
 				TestWidget.Draw();
 			}
diff --git a/Engine/Plugins/Experimental/SlateIM/Source/SlateIM/Private/Misc/SlateIMSlateResources.h b/Engine/Plugins/Experimental/SlateIM/Source/SlateIM/Private/Misc/SlateIMSlateResources.h
--- a/Engine/Plugins/Experimental/SlateIM/Source/SlateIM/Private/Misc/SlateIMSlateResources.h
+++ b/Engine/Plugins/Experimental/SlateIM/Source/SlateIM/Private/Misc/SlateIMSlateResources.h
@@ -2,9 +2,6 @@
  
 #pragma once
 
-// Entire file is WITH_ENGINE
-#if WITH_ENGINE
-
 #include "Styling/SlateBrush.h"
 #include "Types/ISlateMetaData.h"
 #include "UObject/Object.h"
@@ -87,5 +84,3 @@ namespace SlateIM
 		virtual bool IsOfTypeImpl(const FName& Type) const override { return GetTypeId() == Type || Super::IsOfTypeImpl(Type); }
 	};
 }
-
-#endif // WITH_ENGINE
diff --git a/Engine/Plugins/Experimental/SlateIM/Source/SlateIM/Private/Roots/SlateIMViewportRoot.h b/Engine/Plugins/Experimental/SlateIM/Source/SlateIM/Private/Roots/SlateIMViewportRoot.h
--- a/Engine/Plugins/Experimental/SlateIM/Source/SlateIM/Private/Roots/SlateIMViewportRoot.h
+++ b/Engine/Plugins/Experimental/SlateIM/Source/SlateIM/Private/Roots/SlateIMViewportRoot.h
@@ -4,17 +4,13 @@
 
 #if WITH_ENGINE
 #include "ISlateIMRoot.h"
+#include "SlateIMParametersFwd.h"
 #include "UObject/WeakObjectPtrTemplates.h"
 #include "Widgets/Layout/SConstraintCanvas.h"
 
 class SImWrapper;
 class SScaleBox;
 
-namespace SlateIM
-{
-	struct FViewportRootLayout;
-}
-
 class UGameViewportClient;
 class ULocalPlayer;
 
diff --git a/Engine/Plugins/Experimental/SlateIM/Source/SlateIM/Private/SlateIM_Internal.h b/Engine/Plugins/Experimental/SlateIM/Source/SlateIM/Private/SlateIM_Internal.h
--- a/Engine/Plugins/Experimental/SlateIM/Source/SlateIM/Private/SlateIM_Internal.h
+++ b/Engine/Plugins/Experimental/SlateIM/Source/SlateIM/Private/SlateIM_Internal.h
@@ -3,8 +3,9 @@
 #pragma once
 
 #include "Math/MathFwd.h"
+#include "SlateIMParametersFwd.h"
 #include "Templates/FunctionFwd.h"
-#include "Templates//SharedPointerFwd.h"
+#include "Templates/SharedPointerFwd.h"
 
 class SImage;
 enum EOrientation : int;
@@ -22,5 +23,5 @@ namespace SlateIM
 	void BeginWrap(EOrientation Orientation);
 	void EndWrap();
 
-	void Image_Internal(const FSlateColor& ColorAndOpacity, FVector2D DesiredSize, TFunctionRef<const FSlateBrush* (const TSharedRef<SImage>&)> GetBrush);
+	void Image_Internal(const FImageParams& Params, TFunctionRef<const FSlateBrush* (const TSharedRef<SImage>&)> GetBrush);
 } // SlateIM
diff --git a/Engine/Plugins/Experimental/SlateIM/Source/SlateIM/Private/Widgets/SImGraph.cpp b/Engine/Plugins/Experimental/SlateIM/Source/SlateIM/Private/Widgets/SImGraph.cpp
--- a/Engine/Plugins/Experimental/SlateIM/Source/SlateIM/Private/Widgets/SImGraph.cpp
+++ b/Engine/Plugins/Experimental/SlateIM/Source/SlateIM/Private/Widgets/SImGraph.cpp
@@ -2,6 +2,7 @@
 
 
 #include "SImGraph.h"
+#include "SlateIMParameters.h"
 
 SLATE_IMPLEMENT_WIDGET(SImGraph)
 
@@ -35,29 +36,29 @@ void SImGraph::EndGraph()
 	}
 }
 
-void SImGraph::AddLineGraph(const TArrayView<FVector2D>& Points, const FLinearColor& LineColor, float LineThickness, const FDoubleRange& XViewRange, const FDoubleRange& YViewRange)
+void SImGraph::AddLineGraph(const TArrayView<FVector2D>& Points, const SlateIM::FGraphLinePointsParams& Params)
 {
-	ensure(XViewRange.HasLowerBound() && XViewRange.HasUpperBound());
-	ensure(YViewRange.HasLowerBound() && YViewRange.HasUpperBound());
+	ensure(Params.XViewRange.HasLowerBound() && Params.XViewRange.HasUpperBound());
+	ensure(Params.YViewRange.HasLowerBound() && Params.YViewRange.HasUpperBound());
 	
-	FSlateIMLineGraphData& LineGraph = GetNextLineGraph(LineColor, LineThickness, XViewRange, YViewRange);
+	FSlateIMLineGraphData& LineGraph = GetNextLineGraph(Params.LineColor, Params.LineThickness, Params.XViewRange, Params.YViewRange);
 	LineGraph.NormalizedPoints.Reserve(Points.Num());
 
 	for (const FVector2D& Point : Points)
 	{
 		LineGraph.NormalizedPoints.Emplace(
-			(Point.X - XViewRange.GetLowerBound().GetValue()) / (XViewRange.GetUpperBound().GetValue() - XViewRange.GetLowerBound().GetValue()),
-			(Point.Y - YViewRange.GetLowerBound().GetValue()) / (YViewRange.GetUpperBound().GetValue() - YViewRange.GetLowerBound().GetValue()));
+			(Point.X - Params.XViewRange.GetLowerBound().GetValue()) / (Params.XViewRange.GetUpperBound().GetValue() - Params.XViewRange.GetLowerBound().GetValue()),
+			(Point.Y - Params.YViewRange.GetLowerBound().GetValue()) / (Params.YViewRange.GetUpperBound().GetValue() - Params.YViewRange.GetLowerBound().GetValue()));
 	}
 }
 
-void SImGraph::AddLineGraph(const TArrayView<double>& Values, const FLinearColor& LineColor, float LineThickness, const FDoubleRange& ViewRange)
+void SImGraph::AddLineGraph(const TArrayView<double>& Values, const SlateIM::FGraphLineValuesParams& Params)
 {
-	ensure(ViewRange.HasLowerBound() && ViewRange.HasUpperBound());
+	ensure(Params.ViewRange.HasLowerBound() && Params.ViewRange.HasUpperBound());
 
-	const FDoubleRange XViewRange = FDoubleRange(0, static_cast<double>(Values.Num()));
+	const FDoubleRange XViewRange = FDoubleRange(0.0, static_cast<double>(Values.Num()));
 	
-	FSlateIMLineGraphData& LineGraph = GetNextLineGraph(LineColor, LineThickness, XViewRange, ViewRange);
+	FSlateIMLineGraphData& LineGraph = GetNextLineGraph(Params.LineColor, Params.LineThickness, XViewRange, Params.ViewRange);
 	LineGraph.NormalizedPoints.Reserve(Values.Num());
 
 	double XValue = 0.0;
@@ -65,7 +66,7 @@ void SImGraph::AddLineGraph(const TArrayView<double>& Values, const FLinearColor
 	{
 		LineGraph.NormalizedPoints.Emplace(
 			(XValue - XViewRange.GetLowerBound().GetValue()) / (XViewRange.GetUpperBound().GetValue() - XViewRange.GetLowerBound().GetValue()),
-			(Value - ViewRange.GetLowerBound().GetValue()) / (ViewRange.GetUpperBound().GetValue() - ViewRange.GetLowerBound().GetValue()));
+			(Value - Params.ViewRange.GetLowerBound().GetValue()) / (Params.ViewRange.GetUpperBound().GetValue() - Params.ViewRange.GetLowerBound().GetValue()));
 
 		XValue += 1.0;
 	}
@@ -127,5 +128,6 @@ FSlateIMLineGraphData& SImGraph::GetNextLineGraph(const FLinearColor& LineColor,
 	NewLineGraph.LineThickness = LineThickness;
 	NewLineGraph.XViewRange = XViewRange;
 	NewLineGraph.YViewRange = YViewRange;
+
 	return LineGraphs.Emplace_GetRef(MoveTemp(NewLineGraph));
 }
diff --git a/Engine/Plugins/Experimental/SlateIM/Source/SlateIM/Private/Widgets/SImGraph.h b/Engine/Plugins/Experimental/SlateIM/Source/SlateIM/Private/Widgets/SImGraph.h
--- a/Engine/Plugins/Experimental/SlateIM/Source/SlateIM/Private/Widgets/SImGraph.h
+++ b/Engine/Plugins/Experimental/SlateIM/Source/SlateIM/Private/Widgets/SImGraph.h
@@ -2,6 +2,7 @@
 
 #pragma once
 
+#include "SlateIMParametersFwd.h"
 #include "Widgets/SLeafWidget.h"
 
 #define UE_API SLATEIM_API
@@ -21,10 +22,7 @@ class SImGraph : public SLeafWidget
 	SLATE_DECLARE_WIDGET(SImGraph, SLeafWidget)
 	
 public:
-	SLATE_BEGIN_ARGS(SImGraph)
-		{
-		}
-
+	SLATE_BEGIN_ARGS(SImGraph) {}
 	SLATE_END_ARGS()
 
 	UE_API void Construct(const FArguments& InArgs);
@@ -32,8 +30,8 @@ public:
 	UE_API void BeginGraph();
 	UE_API void EndGraph();
 
-	UE_API void AddLineGraph(const TArrayView<FVector2D>& Points, const FLinearColor& LineColor, float LineThickness, const FDoubleRange& XViewRange, const FDoubleRange& YViewRange);
-	UE_API void AddLineGraph(const TArrayView<double>& Points, const FLinearColor& LineColor, float LineThickness, const FDoubleRange& ViewRange);
+	UE_API void AddLineGraph(const TArrayView<FVector2D>& Points, const SlateIM::FGraphLinePointsParams& Params);
+	UE_API void AddLineGraph(const TArrayView<double>& Values, const SlateIM::FGraphLineValuesParams& Params);
 
 private:
 	UE_API virtual int32 OnPaint( const FPaintArgs& Args, const FGeometry& AllottedGeometry, const FSlateRect& MyCullingRect, FSlateWindowElementList& OutDrawElements, int32 LayerId, const FWidgetStyle& InWidgetStyle, bool bParentEnabled ) const override;
diff --git a/Engine/Plugins/Experimental/SlateIM/Source/SlateIM/Public/SlateIM.h b/Engine/Plugins/Experimental/SlateIM/Source/SlateIM/Public/SlateIM.h
--- a/Engine/Plugins/Experimental/SlateIM/Source/SlateIM/Public/SlateIM.h
+++ b/Engine/Plugins/Experimental/SlateIM/Source/SlateIM/Public/SlateIM.h
@@ -1,9 +1,11 @@
 // Copyright Epic Games, Inc. All Rights Reserved.
 
 #pragma once
+
 #include "Framework/Views/ITypedTableView.h"
 #include "Layout/Margin.h"
 #include "Modules/ModuleInterface.h"
+#include "SlateIMParameters.h"
 #include "Styling/SlateColor.h"
 #include "Styling/StyleColors.h"
 #include "Textures/SlateIcon.h"
@@ -61,33 +63,10 @@ public:
 namespace SlateIM
 {
 #pragma region Roots
-#if WITH_EDITOR || WITH_ENGINE
-	struct FViewportRootLayout
-	{
-		/** How to anchor the root within the viewport (normalized units) */
-		FAnchors Anchors = FAnchors(0);
-		
-		/** Offset the root from the anchored position in the viewport (in slate units) */
-		FVector2f Offset = FVector2f::ZeroVector;
-
-		/** How the root is aligned to its anchor (normalized units) */
-		FVector2f Alignment = FVector2f::ZeroVector;
-
-		/** Optional set the size (in slate units) of the root in the viewport. When unset, the root will auto-size for its content. */
-		TOptional<FVector2f> Size;
-
-		/** The ZOrder of the SlateIM widget in the viewport */
-		int32 ZOrder = 10000;
-
-		/** Global scale to apply within the viewport */
-		float Scale = 1.f;
-	};
-#endif // WITH_EDITOR || WITH_ENGINE
-
 	/**
 	 * Begins a new floating window root
 	 *
-	 * Please use the method below with the struct parameters.
+	 * Please use the method below with the struct parameter.
 	 *
 	 * @param UniqueName a unique name to identify this window
 	 * @param WindowTitle The name of the window. This will also be displayed in the windows title bar
@@ -100,26 +79,11 @@ namespace SlateIM
 	 */
 	SLATEIM_API bool BeginWindowRoot(FName UniqueName, const FStringView& WindowTitle, FVector2f WindowSize, bool bShouldReopen = false);
 
-	struct FWindowParams
-	{
-		// The name of the window. This will also be displayed in the windows title bar
-		FStringView WindowTitle;
-
-		// The size to open the window at (does not update the size of an existing window)
-		FVector2f WindowSize = FVector2f::ZeroVector;
-
-		// If this window was created with a previous call to BeginWindow and then closed, passing in true will reopen the window and false will leave it closed. If the window has never been seen this parameter does nothing and the window will open
-		bool bShouldReopen = false;
-
-		// Window will be create always on top.
-		bool bAlwaysOnTop = false;
-	};
-
 	/**
 	 * Begins a new floating window root
 	 *
 	 * @param UniqueName a unique name to identify this window
-	 * @param Params Configuration parameters for the window.
+	 * @param Params (optional) Additional parameters to create the window.
 	 *
 	 * @return the current window state. True if open, false if closed or not updating
 	 *
@@ -129,19 +93,35 @@ namespace SlateIM
 
 #if WITH_ENGINE
 	// TODO - Game viewport widgets might need something to control HitTestability and cursor visibility
+
 	/**
-	 *	Begins a new root in the game viewport
-	 * 
+	 * Begins a new root in the game viewport
+	 *
+	 * Please use the method below with the struct-parameter.
+	 *
 	 * @param UniqueName a globally unique name to identify this root
 	 * @param ViewportClient the viewport client to add the widget to
 	 * @param Layout How to lay out the root within the viewport
-	 * 
+	 *
+	 * @return Whether the root is valid and updating
+	 */
+	SLATEIM_API bool BeginViewportRoot(FName UniqueName, UGameViewportClient* ViewportClient, const FViewportRootLayout& Layout);
+
+	/**
+	 * Begins a new root in the game viewport
+	 *
+	 * @param UniqueName a globally unique name to identify this root
+	 * @param ViewportClient the viewport client to add the widget to
+	 * @param Params (optional) Additional parameters to create the root.
+	 *
 	 * @return Whether the root is valid and updating
 	 */
-	SLATEIM_API bool BeginViewportRoot(FName UniqueName, UGameViewportClient* ViewportClient, const FViewportRootLayout& Layout = FViewportRootLayout());
+	SLATEIM_API bool BeginViewportRoot(FName UniqueName, UGameViewportClient* ViewportClient, const FViewportParams& Params = FViewportParams());
 	
 	/**
-	 *	Begins a new root in the provided player's viewport
+	 * Begins a new root in the provided player's viewport
+	 *
+	 * Please use the method below with the struct-parameter.
 	 * 
 	 * @param UniqueName a globally unique name to identify this root
 	 * @param LocalPlayer the player whose viewport to add the widget to
@@ -149,31 +129,60 @@ namespace SlateIM
 	 * 
 	 * @return Whether the root is valid and updating
 	 */
-	SLATEIM_API bool BeginViewportRoot(FName UniqueName, ULocalPlayer* LocalPlayer, const FViewportRootLayout& Layout = FViewportRootLayout());
+	SLATEIM_API bool BeginViewportRoot(FName UniqueName, ULocalPlayer* LocalPlayer, const FViewportRootLayout& Layout);
+	
+	/**
+	 * Begins a new root in the provided player's viewport
+	 * 
+	 * @param UniqueName a globally unique name to identify this root
+	 * @param LocalPlayer the player whose viewport to add the widget to
+	 * @param Params (optional) Additional parameters to create the root.
+	 * 
+	 * @return Whether the root is valid and updating
+	 */
+	SLATEIM_API bool BeginViewportRoot(FName UniqueName, ULocalPlayer* LocalPlayer, const FViewportParams& Params = FViewportParams());
 #endif
 
 #if WITH_EDITOR
 	/**
 	 * Begins a SlateIM root in the provided in editor viewport
+	 *
+	 * Please use the method below with the struct-parameter.
 	 * 
 	 * @param UniqueName Unique identifier for this root. Must be globally unique within the same editor session
 	 * @param AssetViewport The editor viewpor to add the SlateIM content to
 	 * @param Layout How to lay out the root within the viewport
 	 * 
 	 * @return Whether the root is valid and updating
+	 *
+	 * @note Not available in BPFL: TSharedPtr<IAssetViewport> is not Blueprint-compatible
+	 */
+	SLATEIM_API bool BeginViewportRoot(FName UniqueName, TSharedPtr<IAssetViewport> AssetViewport, const FViewportRootLayout& Layout);
+
+	/**
+	 * Begins a SlateIM root in the provided in editor viewport
+	 *
+	 * @param UniqueName Unique identifier for this root. Must be globally unique within the same editor session
+	 * @param AssetViewport The editor viewpor to add the SlateIM content to
+	 * @param Params (optional) Additional parameters to create the root.
+	 *
+	 * @return Whether the root is valid and updating
+	 *
+	 * @note Not available in BPFL: TSharedPtr<IAssetViewport> is not Blueprint-compatible
 	 */
-	SLATEIM_API bool BeginViewportRoot(FName UniqueName, TSharedPtr<IAssetViewport> AssetViewport, const FViewportRootLayout& Layout = FViewportRootLayout());
+	SLATEIM_API bool BeginViewportRoot(FName UniqueName, TSharedPtr<IAssetViewport> AssetViewport, const FViewportParams& Params = FViewportParams());
 #endif
 
 	/**
 	 * Begins a SlateIM root that exposes a slate widget for embedding in other Slate hierarchies
-	 * 
+	 *
 	 * @param UniqueName a globally unique name to identify this root
 	 * @param OutSlateIMWidget where to output the resulting SlateIM widget to
-	 * 
+	 *
 	 * @return Whether the root is valid and updating
-	 * 
+	 *
 	 * @note Consider using FSlateIMExposedBase to handle creating and updating your Slate IM exposed widget
+	 * @note Not available in BPFL: TSharedPtr<SWidget> is not Blueprint-compatible
 	 */
 	SLATEIM_API bool BeginExposedRoot(FName UniqueName, TSharedPtr<SWidget>& OutSlateIMWidget);
 
@@ -186,51 +195,49 @@ namespace SlateIM
 #pragma region Containers
 	/**
 	 * Begins a horizontally stacked container. All widgets created within the Horizontal Stack container will be placed in-order left-to-right
+	 * 
+	 * @param bMaximizeContent (optional) Maximize() the contents of the stack
 	 */
-	SLATEIM_API void BeginHorizontalStack();
+	SLATEIM_API void BeginHorizontalStack(bool bMaximizeContent = false);
+
 	SLATEIM_API void EndHorizontalStack();
 
 	/**
 	 * Begins a vertically stacked container. All widgets created within the Vertical Stack container will be placed in-order top-to-bottom
+	 *
+	 * @param bMaximizeContent (optional) Maximize() the contents of the stack
 	 */
-	SLATEIM_API void BeginVerticalStack();
+	SLATEIM_API void BeginVerticalStack(bool bMaximizeContent = false);
 	SLATEIM_API void EndVerticalStack();
 
 	/**
 	 * Begins a horizontally wrapped container. All widgets created within the Horizontal Wrap container will be placed in-order left-to-right until the allotted width is filled,
 	 * then the content will begin a new line and so-on.
 	 *
+	 * @param bMaximizeContent (optional) Maximize() the contents of the wrap
+	 *
 	 * @note Wrap box children are always AutoSize and do not support Fill
 	 */
-	SLATEIM_API void BeginHorizontalWrap();
+	SLATEIM_API void BeginHorizontalWrap(bool bMaximizeContent = false);
+
 	SLATEIM_API void EndHorizontalWrap();
 	
 	/**
 	 * Begins a vertically wrapped container. All widgets created within the Vertical Wrap container will be placed in-order top-to-bottom until the allotted height is filled,
 	 * then the content will begin a new column and so-on.
 	 *
+	 * @param bMaximizeContent (optional) Maximize() the contents of the wrap
+	 *
 	 * @note Wrap container children are always AutoSize and do not support Fill
 	 */
-	SLATEIM_API void BeginVerticalWrap();
-	SLATEIM_API void EndVerticalWrap();
+	SLATEIM_API void BeginVerticalWrap(bool bMaximizeContent = false);
 
-	struct FTableParams
-	{
-		// A Table View Style to override the default style. Changes to the style after the first invocation are not reflected,
-		// provide a new style object to update the visuals at runtime
-		const FTableViewStyle* Style = nullptr;
-		
-		// A Table Row Style to override the default table row style. Changing the style only affects new table rows.
-		const FTableRowStyle* RowStyle = nullptr;
-
-		// Select mode for the table.
-		ESelectionMode::Type SelectionMode = ESelectionMode::None;
-	};
+	SLATEIM_API void EndVerticalWrap();
 
 	/**
 	 * Begin a Table container. The contents are automatically scrollable if the table is sized smaller than all its content.
 	 *
-	 * Please use the below method instead!
+	 * Please use the method below with the struct-parameter.
 	 *
 	 * @param InStyle A Table View Style to override the default style. Changes to the style after the first invocation are not reflected,
 	 *				  provide a new style object to update the visuals at runtime
@@ -241,7 +248,7 @@ namespace SlateIM
 	/**
 	 * Begin a Table container. The contents are automatically scrollable if the table is sized smaller than all its content.
 	 *
-	 * @param Params (optional) Settings for table construction.
+	 * @param Params (optional) Additional parameters to create the table.
 	 */
 	SLATEIM_API void BeginTable(const FTableParams& Params = FTableParams());
 
@@ -272,10 +279,20 @@ namespace SlateIM
 	/**
 	 * Adds a column and immediately ends it with the given label as its body.
 	 *
+	 * Please use the method below with the struct-parameter.
+	 *
 	 * @param ColumnID The name to use as ID for the column
 	 * @param ColumnLabel Added a default label to the column
 	 */
-	SLATEIM_API void AddTableColumn(const FName& ColumnID, const FStringView& ColumnLabel = FStringView());
+	SLATEIM_API void AddTableColumn(const FName& ColumnID, const FStringView& ColumnLabel);
+
+	/**
+	 * Adds a column and immediately ends it with the given label as its body.
+	 *
+	 * @param ColumnID The name to use as ID for the column
+	 * @param Params (optional) Additional parameters to create the column
+	 */
+	SLATEIM_API void AddTableColumn(const FName& ColumnID, const FTableColumnParams& Params = FTableColumnParams());
 
 	/**
 	 * Set the column to a fixed width, not resizable by users. Call this before calling AddTableColumn().
@@ -317,6 +334,8 @@ namespace SlateIM
 	 * If the table is a tree, a unique row id must be provided when starting a new set of child rows so that the expansion
 	 * state of the parent row is saved. Not providing a row id will cause the parent row to revert to the default expansion
 	 * state if its grandparent row is closed.
+	 *
+	 * Please use the method below with the struct-parameter.
 	 * 
 	 * @param ParentRowId Id of the parent row. 0 to disable state saving. Row ids must be unique to its parent.
 	 * @param bDefaultExpanded (optional) Set the default expansion state of the parent (whether the children following are visible.)
@@ -324,7 +343,20 @@ namespace SlateIM
 	 * 
 	 * @return Whether the parent row is expanded or not. When false, the child content can be skipped to save cycles.
 	 */
-	SLATEIM_API bool BeginTableRowChildren(uint32 ParentRowId = 0, bool bDefaultExpanded = false);
+	SLATEIM_API bool BeginTableRowChildren(uint32 ParentRowId, bool bDefaultExpanded = false);
+
+	/**
+	 * Begin adding content as a child to the current table row. NextTableCell() must still be called after this function.
+	 * If the table is a tree, a unique row id must be provided when starting a new set of child rows so that the expansion
+	 * state of the parent row is saved. Not providing a row id will cause the parent row to revert to the default expansion
+	 * state if its grandparent row is closed.
+	 *
+	 * @param Params (optional) Additional parameters to create the table children.
+	 *
+	 * @return Whether the parent row is expanded or not. When false, the child content can be skipped to save cycles.
+	 */
+	SLATEIM_API bool BeginTableRowChildren(const FTableRowChildrenParams& Params = FTableRowChildrenParams());
+	
 	/**
 	 * Stop adding child content to the parent row. This must be called no matter the result of BeginTableRowChildren()
 	 */
@@ -332,37 +364,108 @@ namespace SlateIM
 
 	/**
 	 * Begin a container with a background image.
+	 *
+	 * Please use the method below with the struct-parameter.
 	 * 
 	 * @param BorderBrush The brush to use as the background for this container
 	 * @param Orientation Which direction the contents of the container should flow
-	 * @param bAbsorbMouse Whether the container should handle all mouse inputs
-	 * @param ContentPadding How much to pad the contents of the container
+	 * @param bAbsorbMouse (optional) Whether the container should handle all mouse inputs
+	 * @param ContentPadding (optional) How much to pad the contents of the container
 	 */
-	SLATEIM_API void BeginBorder(const FSlateBrush* BorderBrush, EOrientation Orientation = Orient_Vertical, bool bAbsorbMouse = true, FMargin ContentPadding = FMargin(2.0f));
+	SLATEIM_API void BeginBorder(const FSlateBrush* BorderBrush, EOrientation Orientation, bool bAbsorbMouse = true, FMargin ContentPadding = FMargin(2.0f));
 	
 	/**
 	 * Begin a container with a background image.
+	 *
+	 * @param BorderBrush The brush to use as the background for this container
+	 * @param Params (optional) Additional params to create the border
+	 */
+	SLATEIM_API void BeginBorder(const FSlateBrush* BorderBrush, const FBorderParams& Params = FBorderParams());
+
+	/**
+	 * Begin a container with a background image.
+	 *
+	 * Please use the method below with the struct-parameter.
 	 * 
 	 * @param BorderStyleName The name of the brush in FAppStyle to use as the background for this container
 	 * @param Orientation Which direction the contents of the container should flow
-	 * @param bAbsorbMouse Whether the container should handle all mouse inputs
-	 * @param ContentPadding How much to pad the contents of the container
+	 * @param bAbsorbMouse (optional) Whether the container should handle all mouse inputs
+	 * @param ContentPadding (optional) How much to pad the contents of the container
+	 */
+	SLATEIM_API void BeginBorder(const FName BorderStyleName, EOrientation Orientation, bool bAbsorbMouse = true, FMargin ContentPadding = FMargin(2.0f));
+
+	/**
+	 * Begin a container with a background image.
+	 *
+	 * @param BorderStyleName The name of the brush in FAppStyle to use as the background for this container
+	 * @param Params (optional) Additional params to create the border
 	 */
-	SLATEIM_API void BeginBorder(const FName BorderStyleName, EOrientation Orientation = Orient_Vertical, bool bAbsorbMouse = true, FMargin ContentPadding = FMargin(2.0f));
+	SLATEIM_API void BeginBorder(const FName BorderStyleName, const FBorderParams& Params = FBorderParams());
+
 	SLATEIM_API void EndBorder();
 
 	/**
 	 * Begins a container that will allow the user to scroll when its content is larger than the allotted space.
+	 *
+	 * Please use the method below with the struct-parameter.
 	 * 
-	 * @param ScrollBarOrientation Which direction content should flow and scroll
+	 * @param Orientation Which direction content should flow and scroll
 	 * 
 	 * @return Whether the user scroll the content this frame
 	 */
-	SLATEIM_API bool BeginScrollBox(EOrientation ScrollBarOrientation = EOrientation::Orient_Vertical);
+	SLATEIM_API bool BeginScrollBox(EOrientation Orientation);
+
+	/**
+	 * Begins a container that will allow the user to scroll when its content is larger than the allotted space.
+	 *
+	 * @param Params (optional) Additional parameters to create the scroll box
+	 *
+	 * @return Whether the user scroll the content this frame
+	 */
+	SLATEIM_API bool BeginScrollBox(const FScrollBoxParams& Params = FScrollBoxParams());
+
 	SLATEIM_API void EndScrollBox();
 
-	SLATEIM_API void BeginPopUp(const FName BorderStyleName = TEXT("ToolPanel.GroupBorder"), EOrientation Orientation = Orient_Vertical, bool bAbsorbMouse = true, FMargin ContentPadding = FMargin(2.0f));
-	SLATEIM_API void BeginPopUp(const FSlateBrush* BorderBrush, EOrientation Orientation = Orient_Vertical, bool bAbsorbMouse = true, FMargin ContentPadding = FMargin(2.0f));
+	/**
+	 * Begin a container with a background image.
+	 *
+	 * Please use the method below with the struct-parameter.
+	 * 
+	 * @param BorderStyleName The name of the brush in FAppStyle to use as the background for this container
+	 * @param Orientation Which direction the contents of the container should flow
+	 * @param bAbsorbMouse (optional) Whether the container should handle all mouse inputs
+	 * @param ContentPadding (optional) How much to pad the contents of the container
+	 */
+	SLATEIM_API void BeginPopUp(const FName BorderStyleName, EOrientation Orientation, bool bAbsorbMouse = true, FMargin ContentPadding = FMargin(2.0f));
+	
+	/**
+	 * Begin a container with a background image.
+	 *
+	 * @param BorderStyleName The name of the brush in FAppStyle to use as the background for this container
+	 * @param Params (optional) Additional parameters to create the border
+	 */
+	SLATEIM_API void BeginPopUp(const FName BorderStyleName = TEXT("ToolPanel.GroupBorder"), const FPopUpParams& Params = FPopUpParams());
+
+	/**
+	 * Begin a container with a background image.
+	 *
+	 * Please use the method below with the struct-parameter.
+	 *
+	 * @param BorderBrush The brush to use as the background for this container
+	 * @param Orientation Which direction the contents of the container should flow
+	 * @param bAbsorbMouse (optional) Whether the container should handle all mouse inputs
+	 * @param ContentPadding (optional) How much to pad the contents of the container
+	 */
+	SLATEIM_API void BeginPopUp(const FSlateBrush* BorderBrush, EOrientation Orientation, bool bAbsorbMouse = true, FMargin ContentPadding = FMargin(2.0f));
+
+	/**
+	 * Begin a container with a background image.
+	 *
+	 * @param BorderBrush The brush to use as the background for this container
+	 * @param Params (optional) Additional params to create the border
+	 */
+	SLATEIM_API void BeginPopUp(const FSlateBrush* BorderBrush, const FPopUpParams& Params = FPopUpParams());
+	
 	SLATEIM_API void EndPopUp();
 #pragma endregion Containers
 
@@ -391,6 +494,7 @@ namespace SlateIM
 	 * Set the next slot to AutoSize to its content
 	 */
 	SLATEIM_API void AutoSize();
+
 	/**
 	 * Sets the next slot to Fill the remaining space in its container
 	 *
@@ -404,17 +508,20 @@ namespace SlateIM
 	 * @param InMinWidth The minimum width for the next slot
 	 */
 	SLATEIM_API void MinWidth(float InMinWidth);
+
 	/**
 	 * Set the minimum height the next slot can have.
 	 * 
 	 * @param InMinHeight The minimum height for the next slot
 	 */
+
 	SLATEIM_API void MinHeight(float InMinHeight);
 	/**
 	 * Set the maximum width the next slot can have.
 	 * 
 	 * @param InMaxWidth The maximum width for the next slot
 	 */
+
 	SLATEIM_API void MaxWidth(float InMaxWidth);
 	/**
 	 * Set the maximum height the next slot can have.
@@ -422,6 +529,7 @@ namespace SlateIM
 	 * @param InMaxHeight The maximum height for the next slot
 	 */
 	SLATEIM_API void MaxHeight(float InMaxHeight);
+
 	/**
 	 * Set the alignments to Fill and removes autosize.
 	 */
@@ -431,14 +539,19 @@ namespace SlateIM
 #pragma region Widgets
 	/**
 	 * Display a string of text
+	 *
+	 * Please use the method below with the struct-parameter.
 	 * 
 	 * @param InText The text string to display
-	 * @param TextStyle (optional) A Text Block Style to override the default style. Changes to the style after the first invocation are not reflected,
+	 * @param TextStyle A Text Block Style to override the default style. Changes to the style after the first invocation are not reflected,
 	 *					provide a new style object to update the visuals at runtime
 	 */
-	SLATEIM_API void Text(const FStringView& InText, const FTextBlockStyle* TextStyle = nullptr);
+	SLATEIM_API void Text(const FStringView& InText, const FTextBlockStyle* TextStyle);
+
 	/**
 	 * Display a colored string of text
+	 *
+	 * Please use the method below with the struct-parameter.
 	 * 
 	 * @param InText The text string to display
 	 * @param Color The color to display the text as
@@ -447,6 +560,14 @@ namespace SlateIM
 	 */
 	SLATEIM_API void Text(const FStringView& InText, FSlateColor Color, const FTextBlockStyle* TextStyle = nullptr);
 
+	/**
+	 * Display a string of text
+	 *
+	 * @param InText The text string to display
+	 * @param Params (optional) Additional parameters to create the text.
+	 */
+	SLATEIM_API void Text(const FStringView& InText, const FTextParams& Params = FTextParams());
+
 	/**
 	 * Create a text input field
 	 * 
@@ -457,96 +578,185 @@ namespace SlateIM
 	 *					
 	 * @return Whether the user changed or committed text this frame
 	 */
-	SLATEIM_API bool EditableText(FString& InOutText, const FStringView& HintText = FStringView(), const FEditableTextBoxStyle* TextStyle = nullptr);
+	SLATEIM_API bool EditableText(FString& InOutText, const FStringView& HintText, const FEditableTextBoxStyle* TextStyle = nullptr);
+
+	/**
+	 * Create a text input field
+	 *
+	 * @param InOutText The text that is in the input field. Can have a default value
+	 * @param Params (optional) Additional parameters to create the widget
+	 *
+	 * @return Whether the user changed or committed text this frame
+	 */
+	SLATEIM_API bool EditableText(FString& InOutText, const FEditableTextParams& Params = FEditableTextParams());
 
 	/**
 	 * Display a brush
 	 * 
 	 * @param ImageBrush The brush to display
-	 * @param ColorAndOpacity (optional) The color to tint the brush
+	 * @param ColorAndOpacity The color to tint the brush
 	 * @param DesiredSize (optional) Override the desired size to display the brush at
 	 */
-	SLATEIM_API void Image(const FSlateBrush* ImageBrush, const FSlateColor& ColorAndOpacity = FLinearColor::White, FVector2D DesiredSize = FVector2D::ZeroVector);
+	SLATEIM_API void Image(const FSlateBrush* ImageBrush, const FSlateColor& ColorAndOpacity, FVector2D DesiredSize = FVector2D::ZeroVector);
+
+	/**
+	 * Display a brush
+	 *
+	 * @param ImageBrush The brush to display
+	 * @param Params (optional) Additional parameters to display the image
+	 */
+	SLATEIM_API void Image(const FSlateBrush* ImageBrush, const FImageParams& Params = FImageParams());
+
 	/**
 	 * Display a slate style brush
+	 *
+	 * Please use the method below with the struct-parameter.
 	 * 
 	 * @param ImageStyleName The named of the brush to display from FAppStyle
-	 * @param ColorAndOpacity (optional) The color to tint the brush
+	 * @param ColorAndOpacity The color to tint the brush
 	 * @param DesiredSize (optional) Override the desired size to display the brush at
 	 */
-	SLATEIM_API void Image(const FName ImageStyleName, const FSlateColor& ColorAndOpacity = FLinearColor::White, FVector2D DesiredSize = FVector2D::ZeroVector);
+	SLATEIM_API void Image(const FName ImageStyleName, const FSlateColor& ColorAndOpacity, FVector2D DesiredSize = FVector2D::ZeroVector);
+	
+	/**
+	 * Display a slate style brush
+	 *
+	 * @param ImageStyleName The named of the brush to display from FAppStyle
+	 * @param Params (optional) Additional parameters to display the image
+	 */
+	SLATEIM_API void Image(const FName ImageStyleName, const FImageParams& Params = FImageParams());
+
 	/**
 	 * Display a colored box
+	 *
+	 * Please use the method below with the struct-parameter.
 	 * 
 	 * @param ColorAndOpacity The color to display
 	 * @param DesiredSize (optional) Override the desired size to display the box at
 	 */
 	SLATEIM_API void Image(const FSlateColor& ColorAndOpacity, FVector2D DesiredSize = FVector2D::ZeroVector);
+
+	/**
+	 * Display a colored box
+	 *
+	 * @param Params (optional) Additional parameters to display the block (including the color)
+	 */
+	SLATEIM_API void Image(const FImageParams& Params = FImageParams());
+
 #if WITH_ENGINE
 	/**
 	 * Display a texture
+	 *
+	 * Please use the method below with the struct-parameter.
 	 * 
 	 * @param ImageTexture The texture to display
 	 * @param ColorAndOpacity (optional) The color to tint the brush
 	 * @param DesiredSize (optional) Override the desired size to display the brush at
 	 */
-	SLATEIM_API void Image(UTexture2D* ImageTexture, const FSlateColor& ColorAndOpacity = FLinearColor::White, FVector2D DesiredSize = FVector2D::ZeroVector);
+	
+	SLATEIM_API void Image(UTexture2D* ImageTexture, const FSlateColor& ColorAndOpacity, FVector2D DesiredSize = FVector2D::ZeroVector);
+	
+	/**
+	 * Display a texture
+	 *
+	 * @param ImageTexture The texture to display
+	 * @param Params (optional) Additional parameters to display the image
+	 */
+	SLATEIM_API void Image(UTexture2D* ImageTexture, const FImageParams& Params = FImageParams());
+
 	/**
 	 * Display a render target texture
+	 *
+	 * Please use the method below with the struct-parameter.
 	 * 
 	 * @param ImageRenderTarget The render target texture to display
-	 * @param ColorAndOpacity (optional) The color to tint the brush
+	 * @param ColorAndOpacity The color to tint the brush
 	 * @param DesiredSize (optional) Override the desired size to display the brush at
 	 */
-	SLATEIM_API void Image(UTextureRenderTarget2D* ImageRenderTarget, const FSlateColor& ColorAndOpacity = FLinearColor::White, FVector2D DesiredSize = FVector2D::ZeroVector);
+	
+	SLATEIM_API void Image(UTextureRenderTarget2D* ImageRenderTarget, const FSlateColor& ColorAndOpacity, FVector2D DesiredSize = FVector2D::ZeroVector);
+	
+	/**
+	 * Display a render target texture
+	 *
+	 * @param ImageRenderTarget The render target texture to display
+	 * @param Params (optional) Additional parameters to display the image
+	 */
+	SLATEIM_API void Image(UTextureRenderTarget2D* ImageRenderTarget, const FImageParams& Params = FImageParams());
+
 	/**
 	 * Display a material
+	 *
+	 * Please use the method below with the struct-parameter.
 	 * 
 	 * @param ImageMaterial The material to display
 	 * @param BrushSize The size to create the internal brush with
-	 * @param ColorAndOpacity (optional) The color to tint the brush
+	 * @param ColorAndOpacity The color to tint the brush
 	 * @param DesiredSize (optional) Override the desired size to display the brush at
 	 */
-	SLATEIM_API void Image(UMaterialInterface* ImageMaterial, FVector2D BrushSize, const FSlateColor& ColorAndOpacity = FLinearColor::White, FVector2D DesiredSize = FVector2D::ZeroVector);
+	SLATEIM_API void Image(UMaterialInterface* ImageMaterial, FVector2D BrushSize, const FSlateColor& ColorAndOpacity, FVector2D DesiredSize = FVector2D::ZeroVector);
+
+	/**
+	 * Display a material
+	 *
+	 * @param ImageMaterial The material to display
+	 * @param BrushSize The size to create the internal brush with
+	 * @param Params (optional) Additional parameters to display the image
+	 */
+	SLATEIM_API void Image(UMaterialInterface* ImageMaterial, FVector2D BrushSize, const FImageParams& Params = FImageParams());
 #endif
 
 	/**
 	 * Display a button with text
 	 * 
+	 * Please use the method below with the struct-parameter.
+
 	 * @param InText The text to display on the button
+	 * @param bEnabled Indicate if the button is enabled or not (grayed)
 	 * @param InStyle (optional) A Button Style to override the default style. Changes to the style after the first invocation are not reflected,
 	 *				  provide a new style object to update the visuals at runtime
 	 *					
 	 * @return Whether the user clicked the button this frame
 	 */
-	SLATEIM_API bool Button(const FStringView& InText, const FButtonStyle* InStyle = nullptr);
+	SLATEIM_API bool Button(const FStringView& InText, const bool bEnabled, const FButtonStyle* InStyle = nullptr);
 
 	/**
-	 * Display a button with text
-	 * 
+	 *
 	 * @param InText The text to display on the button
-	 * @param bEnabled Indicate if the button is enabled or not (grayed)
-	 * @param InStyle (optional) A Button Style to override the default style. Changes to the style after the first invocation are not reflected,
-	 *				  provide a new style object to update the visuals at runtime
-	 *					
+	 * @param Params (optional) Additional parameters to create the button
+	 *
 	 * @return Whether the user clicked the button this frame
 	 */
-	SLATEIM_API bool Button(const FStringView& InText, const bool bEnabled, const FButtonStyle* InStyle = nullptr);
+	SLATEIM_API bool Button(const FStringView& InText, const FButtonParams& Params = FButtonParams());
 
 	/**
 	 * Display a two-state checkbox
+	 *
+	 * Please use the method below with the struct-parameter.
 	 * 
-	 * @param InText The label for the checkbox
+	 * @param InLabel The label for the checkbox
 	 * @param InOutCurrentState The current state of the checkbox. Can provide a default value.
 	 * @param CheckBoxStyle (optional) A Checkbox Style to override the default style. Changes to the style after the first invocation are not reflected,
 	 *						provide a new style object to update the visuals at runtime
 	 *					
 	 * @return Whether the user changed the state of the checkbox this frame
 	 */
-	SLATEIM_API bool CheckBox(const FStringView& InText, bool& InOutCurrentState, const FCheckBoxStyle* CheckBoxStyle = nullptr);
+	SLATEIM_API bool CheckBox(const FStringView& InLabel, bool& InOutCurrentState, const FCheckBoxStyle* CheckBoxStyle = nullptr);
+
+	/**
+	 * Display a two-state checkbox
+	 *
+	 * @param InOutCurrentState The current state of the checkbox. Can provide a default value.
+	 * @param Params (optional) Additional parameters to create the checkbox
+	 *
+	 * @return Whether the user changed the state of the checkbox this frame
+	 */
+	SLATEIM_API bool CheckBox(bool& InOutCurrentState, const FCheckBoxParams& Params = FCheckBoxParams());
 	
 	/**
 	 * Display a three-state checkbox
+	 *
+	 * Please use the method below with the struct-parameter.
 	 * 
 	 * @param InText The label for the checkbox
 	 * @param InOutCurrentState The current state of the checkbox. Can provide a default value.
@@ -557,8 +767,20 @@ namespace SlateIM
 	 */
 	SLATEIM_API bool CheckBox(const FStringView& InText, ECheckBoxState& InOutCurrentState, const FCheckBoxStyle* CheckBoxStyle = nullptr);
 
+	/**
+	 * Display a three-state checkbox
+	 *
+	 * @param InOutCurrentState The current state of the checkbox. Can provide a default value.
+	 * @param Params (optional) Additional parameters to create the checkbox
+	 *
+	 * @return Whether the user changed the state of the checkbox this frame
+	 */
+	SLATEIM_API bool CheckBox(ECheckBoxState& InOutCurrentState, const FCheckBoxParams& Params = FCheckBoxParams());
+
 	/**
 	 * Display a float-based spin box
+	 *
+	 * Please use the method below with the struct-parameter.
 	 * 
 	 * @param InOutValue The current value of the spin box. Can provide a default value.
 	 * @param Min The minimum value of the spin box (or unset for no limit).
@@ -569,8 +791,11 @@ namespace SlateIM
 	 * @return Whether the user changed or committed the spin box value this frame
 	 */
 	SLATEIM_API bool SpinBox(float& InOutValue, TOptional<float> Min, TOptional<float> Max, const FSpinBoxStyle* SpinBoxStyle = nullptr);
+
 	/**
 	 * Display a double-based spin box
+	 *
+	 * Please use the method below with the struct-parameter.
 	 * 
 	 * @param InOutValue The current value of the spin box. Can provide a default value.
 	 * @param Min The minimum value of the spin box (or unset for no limit).
@@ -581,8 +806,11 @@ namespace SlateIM
 	 * @return Whether the user changed or committed the spin box value this frame
 	 */
 	SLATEIM_API bool SpinBox(double& InOutValue, TOptional<double> Min, TOptional<double> Max, const FSpinBoxStyle* SpinBoxStyle = nullptr);
+
 	/**
 	 * Display an integer-based spin box
+	 *
+	 * Please use the method below with the struct-parameter.
 	 * 
 	 * @param InOutValue The current value of the spin box. Can provide a default value.
 	 * @param Min The minimum value of the spin box (or unset for no limit).
@@ -595,30 +823,96 @@ namespace SlateIM
 	SLATEIM_API bool SpinBox(int32& InOutValue, TOptional<int32> Min, TOptional<int32> Max, const FSpinBoxStyle* SpinBoxStyle = nullptr);
 
 	/**
-	 * Display a float-based slider
+	 * Display a spin box
+	 *
+	 * @param InOutValue The current value of the spin box. Can provide a default value.
+	 * @param Params (optional) Additional parameters to create the spinbox
+	 *
+	 * @return Whether the user committed the spin box value this frame
+	 */
+	template<typename InNumericType>
+	bool SpinBox(InNumericType& InOutValue, const FSpinBoxParams<InNumericType>& Params = FSpinBoxParams<InNumericType>());
+
+	/**
+	 * Display a float-based spin box
+	 *
+	 * @param InOutValue The current value of the spin box. Can provide a default value.
+	 * @param Params Additional parameters to create the spinbox
+	 * 
+	 * @return Whether the user changed or committed the spin box value this frame
+	 */
+	template<>
+	bool SLATEIM_API SpinBox(float& InOutValue, const FSpinBoxParams<float>& Params);
+
+	/**
+	 * Display a double-based spin box
+	 *
+	 * @param InOutValue The current value of the spin box. Can provide a default value.
+	 * @param Params Additional parameters to create the spinbox
+	 *
+	 * @return Whether the user changed or committed the spin box value this frame
+	 */
+	template<>
+	bool SLATEIM_API SpinBox(double& InOutValue, const FSpinBoxParams<double>& Params);
+
+	/**
+	 * Display an int32-based spin box
+	 *
+	 * @param InOutValue The current value of the spin box. Can provide a default value.
+	 * @param Params Additional parameters to create the spinbox
 	 * 
+	 * @return Whether the user changed or committed the spin box value this frame
+	 */
+	template<>
+	bool SLATEIM_API SpinBox(int32& InOutValue, const FSpinBoxParams<int32>& Params);
+
+	/**
+	 * Display a float-based slider
+	 *
+	 * Please use the method below with the struct-parameter.
+	 *
 	 * @param InOutValue The current value of the slider. Can provide a default value.
 	 * @param Min The minimum value of the slider
 	 * @param Max The maximum value of the slider
 	 * @param Step The smallest incremental change that can be made to the slider's value
 	 * @param SliderStyle (optional) A Slider Style to override the default style. Changes to the style after the first invocation are not reflected,
 	 *					  provide a new style object to update the visuals at runtime
-	 *					  
+	 *
 	 * @return Whether the user changed the value of the slider this frame
 	 */
 	SLATEIM_API bool Slider(float& InOutValue, float Min, float Max, float Step, const FSliderStyle* SliderStyle = nullptr);
 
+	/**
+	 * Display a float-based slider
+	 *
+	 * @param InOutValue The current value of the slider. Can provide a default value.
+	 * @param Params (optional) Additional parameters to create the slider
+	 *
+	 * @return Whether the user changed the value of the slider this frame
+	 */
+	SLATEIM_API bool Slider(float& InOutValue, const FSliderParams& Params = FSliderParams());
+
 	/**
 	 * Display a progress bar
 	 * 
 	 * @param Percent The value of the progress bar (0.0 to 1.0)
-	 * @param ProgressBarStyle (optional) A Progress Bar Style to override the default style. Changes to the style after the first invocation are not reflected,
+	 * @param ProgressBarStyle A Progress Bar Style to override the default style. Changes to the style after the first invocation are not reflected,
 	 *						   provide a new style object to update the visuals at runtime
 	 */
-	SLATEIM_API void ProgressBar(TOptional<float> Percent, const FProgressBarStyle* ProgressBarStyle = nullptr);
+	SLATEIM_API void ProgressBar(TOptional<float> Percent, const FProgressBarStyle* ProgressBarStyle);
+
+	/**
+	 * Display a progress bar
+	 *
+	 * @param Percent The value of the progress bar (0.0 to 1.0)
+	 * @param Params (optional) Additional parameters to create the progress bar
+	 */
+	SLATEIM_API void ProgressBar(TOptional<float> Percent, const FProgressBarParams& Params = FProgressBarParams());
 
 	/**
 	 * Display a dropdown of text options
+	 *
+	 * Please use the method below with the struct-parameter.
 	 * 
 	 * @param ComboItems The options the user can choose from
 	 * @param InOutSelectedItemIndex The index of the currently selected option
@@ -628,10 +922,12 @@ namespace SlateIM
 	 *					   
 	 * @return Whether the user changed the selected option this frame
 	 */
-	SLATEIM_API bool ComboBox(const TArray<FString>& ComboItems, int32& InOutSelectedItemIndex, bool bForceRefresh = false, const FComboBoxStyle* InComboStyle = nullptr);
+	SLATEIM_API bool ComboBox(const TArray<FString>& ComboItems, int32& InOutSelectedItemIndex, bool bForceRefresh, const FComboBoxStyle* InComboStyle = nullptr);
 
 	/**
 	 * Display a searchable dropdown of text options
+	 *
+	 * Please use the method below with the struct-parameter.
 	 * 
 	 * @param ComboItems The options the user can choose from
 	 * @param InOutSelectedItemIndex The index of the currently selected option
@@ -642,8 +938,21 @@ namespace SlateIM
 	 */
 	SLATEIM_API bool SearchableComboBox(const TArray<FString>& ComboItems, int32& InOutSelectedItemIndex, bool bForceRefresh = false, const FComboBoxStyle* InComboStyle = nullptr);
 
+	/**
+	 * Display a dropdown of text options
+	 *
+	 * @param ComboItems The options the user can choose from
+	 * @param InOutSelectedItemIndex The index of the currently selected option
+	 * @param Params (optional) Additional parameter to create the combo box
+	 *
+	 * @return Whether the user changed the selected option this frame
+	 */
+	SLATEIM_API bool ComboBox(const TArray<FString>& ComboItems, int32& InOutSelectedItemIndex, const FComboBoxParams& Params = FComboBoxParams());
+
 	/**
 	 * Display a list of text options
+	 *
+	 * Please use the method below with the struct-parameter.
 	 * 
 	 * @param ListItems The list of options the user can choose from
 	 * @param InOutSelectedItemIndex The index of the currently selected option
@@ -653,7 +962,18 @@ namespace SlateIM
 	 *					   
 	 * @return Whether the user changed the selected option this frame
 	 */
-	SLATEIM_API bool SelectionList(const TArray<FString>& ListItems, int32& InOutSelectedItemIndex, bool bForceRefresh = false, const FTableViewStyle* InStyle = nullptr);
+	SLATEIM_API bool SelectionList(const TArray<FString>& ListItems, int32& InOutSelectedItemIndex, bool bForceRefresh, const FTableViewStyle* InStyle = nullptr);
+
+	/**
+	 * Display a list of text options
+	 *
+	 * @param ListItems The list of options the user can choose from
+	 * @param InOutSelectedItemIndex The index of the currently selected option
+	 * @param Params (optional) Additional parameters to create the selection list
+	 *
+	 * @return Whether the user changed the selected option this frame
+	 */
+	SLATEIM_API bool SelectionList(const TArray<FString>& ListItems, int32& InOutSelectedItemIndex, const FSelectionListParams& Params = FSelectionListParams());
 
 	/**
 	 * Create a block of empty space
@@ -662,10 +982,11 @@ namespace SlateIM
 	 */
 	SLATEIM_API void Spacer(const FVector2D& Size);
 
-	/** 
+	/**
 	 * Add an already created slate widget to the IM hierarchy
-	 * 
+	 *
 	 * @note Do not create the widget you are passing in here every frame. This is meant to be used to add already created widgets
+	 * @note Not available in BPFL: TSharedRef<SWidget> is not Blueprint-compatible
 	 */
 	SLATEIM_API void Widget(TSharedRef<SWidget> InWidget);
 #pragma endregion Widgets
@@ -677,6 +998,7 @@ namespace SlateIM
 	 * @return Whether the menu is open
 	 */
 	SLATEIM_API bool BeginContextMenuAnchor();
+
 	SLATEIM_API void EndContextMenuAnchor();
 
 	/** 
@@ -715,33 +1037,73 @@ namespace SlateIM
 
 	/**
 	 * Add a menu item button to the current context menu
+	 *
+	 * Please use the method below with the struct-parameter.
 	 * 
 	 * @param RowText Label for the menu item
-	 * @param ToolTipText (optional) The tooltip to display for the menu item
+	 * @param ToolTipText The tooltip to display for the menu item
 	 * 
 	 * @return Whether the user clicked the menu item this frame
 	 */
-	SLATEIM_API bool AddMenuButton(const FStringView& RowText, const FStringView& ToolTipText = FStringView());
+	SLATEIM_API bool AddMenuButton(const FStringView& RowText, const FStringView& ToolTipText);
+
+	/**
+	 * Add a menu item button to the current context menu
+	 *
+	 * @param RowText Label for the menu item
+	 * @param Params (optional) Additional parameters to create the menu button
+	 *
+	 * @return Whether the user clicked the menu item this frame
+	 */
+	SLATEIM_API bool AddMenuButton(const FStringView& RowText, const FMenuButtonParams& Params = FMenuButtonParams());
+
 	/**
 	 * Display a menu item button with a checkbox
+	 *
+	 * Please use the method below with the struct-parameter.
 	 * 
 	 * @param RowText Label for the menu item
 	 * @param InOutCurrentState The current state of the toggle
-	 * @param ToolTipText (optional) The tooltip to display for the menu item
+	 * @param ToolTipText The tooltip to display for the menu item
 	 * 
 	 * @return Whether the user toggled the checkbox this frame
 	 */
-	SLATEIM_API bool AddMenuToggleButton(const FStringView& RowText, bool& InOutCurrentState, const FStringView& ToolTipText = FStringView());
+	SLATEIM_API bool AddMenuToggleButton(const FStringView& RowText, bool& InOutCurrentState, const FStringView& ToolTipText);
+
+	/**
+	 * Display a menu item button with a checkbox
+	 *
+	 * @param RowText Label for the menu item
+	 * @param InOutCurrentState The current state of the toggle
+	 * @param Params (optional) Additional parameters to create the menu button
+	 *
+	 * @return Whether the user toggled the checkbox this frame
+	 */
+	SLATEIM_API bool AddMenuToggleButton(const FStringView& RowText, bool& InOutCurrentState, const FMenuButtonParams& Params = FMenuButtonParams());
+
 	/**
 	 * Display a menu item button with a checkmark
+	 *
+	 * Please use the method below with the struct-parameter.
 	 * 
 	 * @param RowText Label for the menu item
 	 * @param InOutCurrentState The current state of the checkmark
-	 * @param ToolTipText (optional) The tooltip to display for the menu item
+	 * @param ToolTipText The tooltip to display for the menu item
 	 * 
 	 * @return Whether the user clicked the menu item this frame
 	 */
-	SLATEIM_API bool AddMenuCheckButton(const FStringView& RowText, bool& InOutCurrentState, const FStringView& ToolTipText = FStringView());
+	SLATEIM_API bool AddMenuCheckButton(const FStringView& RowText, bool& InOutCurrentState, const FStringView& ToolTipText);
+
+	/**
+	 * Display a menu item button with a checkmark
+	 *
+	 * @param RowText Label for the menu item
+	 * @param InOutCurrentState The current state of the checkmark
+	 * @param Params (optional) Additional parameters to create the menu button
+	 *
+	 * @return Whether the user clicked the menu item this frame
+	 */
+	SLATEIM_API bool AddMenuCheckButton(const FStringView& RowText, bool& InOutCurrentState, const FMenuButtonParams& Params = FMenuButtonParams());
 
 	/**
 	 * Add a submenu item to the current menu
@@ -749,6 +1111,7 @@ namespace SlateIM
 	 * @param SubMenuText The label for the menu item
 	 */
 	SLATEIM_API void BeginSubMenu(const FStringView& SubMenuText);
+
 	SLATEIM_API void EndSubMenu();
 #pragma endregion Menu
 
@@ -796,14 +1159,26 @@ namespace SlateIM
 
 	/**
 	 * Begins a Tab to contain any other content
+	 *
+	 * Please use the method below with the struct-parameter.
 	 * 
 	 * @param TabId Provide a tab id that is unique within the Tab Group
-	 * @param TabIcon (optional) An icon to display on the tab next to the title. This cannot be updated after the initial creation of the tab.
+	 * @param TabIcon An icon to display on the tab next to the title. This cannot be updated after the initial creation of the tab.
 	 * @param TabTitle (optional) The title to display for the tab if it should differ from the TabId. This cannot be updated after the initial creation of the tab.
 	 * 
 	 * @return Returns true if the tab is active, false otherwise. Logic to draw the contents of the tab can be skipped when this is false.
 	 */
-	SLATEIM_API bool BeginTab(const FName& TabId, const FSlateIcon& TabIcon = FSlateIcon(), const FText& TabTitle = FText::GetEmpty());
+	SLATEIM_API bool BeginTab(const FName& TabId, const FSlateIcon& TabIcon, const FText& TabTitle = FText::GetEmpty());
+
+	/**
+	 * Begins a Tab to contain any other content
+	 *
+	 * @param TabId Provide a tab id that is unique within the Tab Group
+	 * @param Params (optional) Additional parameters used to create the tab
+	 *
+	 * @return Returns true if the tab is active, false otherwise. Logic to draw the contents of the tab can be skipped when this is false.
+	 */
+	SLATEIM_API bool BeginTab(const FName& TabId, const FTabParams& Params = FTabParams());
 
 	/**
 	 * Ends the current Tab. This must be called even when BeginTab returns false.
@@ -836,6 +1211,7 @@ namespace SlateIM
 	 * @note It is not possible to enable a widget inside a disabled parent widget by calling EndDisabledState before a child is created inside a disabled Widget.
 	 */
 	SLATEIM_API void BeginDisabledState();
+
 	SLATEIM_API void EndDisabledState();
 
 	/**
@@ -847,36 +1223,73 @@ namespace SlateIM
 
 	/**
 	 * Opens a modal dialog of the specified type. This function will not return until the user closes the modal dialog
+	 *
+	 * Please use the method below with the struct-parameter.
 	 * 
 	 * @param MessageType The type of options the user can respond to the dialog with
 	 * @param DialogText The message to display to the user
 	 * @param Category The type of dialog to display
-	 * @param DialogTitle The title to display in the dialog window
+	 * @param DialogTitle (optional) The title to display in the dialog window
 	 * 
 	 * @return The option the user selected
 	 *
 	 * @note SlateIM updates are disable while a SlateIM modal is open
 	 * @see CanUpdateSlateIM()
 	 */
-	SLATEIM_API EAppReturnType::Type ModalDialog(EAppMsgType::Type MessageType, const FStringView& DialogText, EAppMsgCategory Category = EAppMsgCategory::Warning, const FStringView& DialogTitle = FStringView());
+	SLATEIM_API EAppReturnType::Type ModalDialog(EAppMsgType::Type MessageType, const FStringView& DialogText, EAppMsgCategory Category, const FStringView& DialogTitle = FStringView());
+
+	/**
+	 * Opens a modal dialog of the specified type. This function will not return until the user closes the modal dialog
+	 *
+	 * @param MessageType The type of options the user can respond to the dialog with
+	 * @param DialogText The message to display to the user
+	 * @param Params (optional) Additional parameters used to create the dialog
+	 *
+	 * @return The option the user selected
+	 *
+	 * @note SlateIM updates are disable while a SlateIM modal is open
+	 * @see CanUpdateSlateIM()
+	 */
+	SLATEIM_API EAppReturnType::Type ModalDialog(EAppMsgType::Type MessageType, const FStringView& DialogText, const FModalDialogParams& Params = FModalDialogParams());
 
 	/**
 	 * Text filter input box that updates the filter when text changes.
 	 *
 	 * @param InFilter The text filter to display and update
-	 * @param HintText Optional hint text to display when the input is empty
+	 * @param Params (optional) Additional parameters to create the text filter
+	 *
 	 * @return Whether the user changed or committed text this frame
+	 *
+	 * @note Not available in BPFL: template function
 	 */
 	template<typename ItemType>
-	bool TextFilter(FSlateImTextFilter<ItemType>& InFilter, const FStringView& HintText = FStringView())
+	bool TextFilter(FSlateImTextFilter<ItemType>& InFilter, const FEditableTextParams& Params = FEditableTextParams())
 	{
-		if (EditableText(InFilter.InputText, HintText) && InFilter.TextFilter.IsValid())
+		if (EditableText(InFilter.InputText, Params) && InFilter.TextFilter.IsValid())
 		{
 			InFilter.TextFilter->SetRawFilterText(FText::FromString(InFilter.InputText));
 			return true;
 		}
 		return false;
 	}
+
+	/**
+	 * Text filter input box that updates the filter when text changes.
+	 *
+	 * Please use the method above with the struct-parameter.
+	 *
+	 * @param InFilter The text filter to display and update
+	 * @param HintText hint text to display when the input is empty
+	 * 
+	 * @return Whether the user changed or committed text this frame
+	 *
+	 * @note Not available in BPFL: template function
+	 */
+	template<typename ItemType>
+	bool TextFilter(FSlateImTextFilter<ItemType>& InFilter, const FStringView& HintText)
+	{
+		return TextFilter<ItemType>(InFilter, {.HintText = HintText});
+	}
 #pragma endregion Util
 
 #pragma region Queries
@@ -887,15 +1300,6 @@ namespace SlateIM
 	 */
 	SLATEIM_API bool IsHovered();
 
-	enum class EFocusDepth : uint8
-	{
-		// Only check if the previous widget is focused
-		SelfOnly,
-
-		// Check if the previous widget or any of its childs widgets are focused (recursively)
-		IncludingDescendants
-	};
-
 	/**
 	 * Query whether the previous widget is focused
 	 * 
@@ -918,6 +1322,8 @@ namespace SlateIM
 
 	/**
 	 * Add a line graph of 2D vectors to the graph widget
+	 *
+	 * Please use the method below with the struct-parameter.
 	 * 
 	 * @param Points The X,Y points to plot on the graph
 	 * @param LineColor The color of the line to draw for this graph
@@ -927,9 +1333,19 @@ namespace SlateIM
 	 */
 	SLATEIM_API void GraphLine(const TArrayView<FVector2D>& Points, const FLinearColor& LineColor, float LineThickness, const FDoubleRange& XViewRange, const FDoubleRange& YViewRange);
 
+	/**
+	 * Add a line graph of 2D vectors to the graph widget
+	 *
+	 * @param Points The X,Y points to plot on the graph
+	 * @param Params Additional parameters to draw the line
+	 */
+	SLATEIM_API void GraphLine(const TArrayView<FVector2D>& Points, const FGraphLinePointsParams& Params = FGraphLinePointsParams());
+
 	/**
 	 * Add a line graph of values to the graph widget.
 	 * The value index is used as the X-value when plotting, the graph will scale horizontally to fit all values in the array.
+	 *
+	 * Please use the method below with the struct-parameter.
 	 * 
 	 * @param Values The Y values to plot on the graph
 	 * @param LineColor The color of the line to draw for this graph
@@ -937,6 +1353,15 @@ namespace SlateIM
 	 * @param ViewRange The min and max Y values to vertically scale the graph to
 	 */
 	SLATEIM_API void GraphLine(const TArrayView<double>& Values, const FLinearColor& LineColor, float LineThickness, const FDoubleRange& ViewRange);
+
+	/**
+	 * Add a line graph of values to the graph widget.
+	 * The value index is used as the X-value when plotting, the graph will scale horizontally to fit all values in the array.
+	 *
+	 * @param Values The Y values to plot on the graph
+	 * @param Params Additional parameters to draw the line
+	 */
+	SLATEIM_API void GraphLine(const TArrayView<double>& Values, const FGraphLineValuesParams& Params = FGraphLineValuesParams());
 #pragma endregion Graphing
 
 #pragma region Inputs
@@ -985,64 +1410,41 @@ namespace SlateIM
 #if WITH_ENGINE
 	namespace Canvas
 	{
-		enum class ECanvasUpdateType : uint8
-		{
-			/** Redraw the canvas every frame. */
-			EveryFrame,
-
-			/**
-			 * Only redraw when specifically invalidated.
-			 * Warning: Textures and materials may not be rendered correctly if they are not loaded or shaders are not cached.
-			 */
-			Invalidation
-
-			/** TODO: Automatically redraw when the hash of the draw queue changes. */
-			// HashChange
-		};
-
-		struct FCanvasParams
-		{
-			FSlateColor ColorAndOpacity = FSlateColor(EStyleColor::White);
-			FVector2D DesiredSize = FVector2D::ZeroVector;
-			UObject* WorldContext = nullptr;
-			ECanvasUpdateType UpdateType = ECanvasUpdateType::EveryFrame;
-		};
-
 		/**
 		 * Begin queueing commands to draw to a canvas render target.
 		 * 
 		 * @param Width The width of the canvas in pixels.
 		 * @param Height The height of the canvas in pixels.
-		 * @param ColorAndOpacity (optional) Tint applied to the resulting image.
-		 * @param DesiredSize (optional) The size of the image drawn in the UI. This may be different from the size of the canvas.
-		 * @param WorldContext (optional) The world related to the canvas being drawn, if any.
+		 * @param Params Additional parameters to create the canvas.
 		 */
-		void BeginCanvas(int32 Width, int32 Height, const FCanvasParams& CanvasParams = FCanvasParams());
+		SLATEIM_API void BeginCanvas(int32 Width, int32 Height, const FCanvasParams& CanvasParams = FCanvasParams());
 
 		/**
 		 * Ends the currently drawn canvas and executes all queued commands.
 		 */
-		void EndCanvas();
+		SLATEIM_API void EndCanvas();
 
 		/**
 		 * Invalidate the current canvas, forcing a redraw if its UpdateType is set to Invalidation.
 		 */
-		void Invalidate();
+		SLATEIM_API void Invalidate();
 
 		/** Sets the position of the lower-right corner of the clipping region of the canvas */
-		void SetClip(const FVector2f& ClipPosition);
+		SLATEIM_API void SetClip(const FVector2f& ClipPosition);
 
 		/** Sets the draw color of the canvas. */
-		void SetDrawColor(const FLinearColor& Color);
-		void SetDrawColor(const FColor& Color);
+		SLATEIM_API void SetDrawColor(const FLinearColor& Color);
+		SLATEIM_API void SetDrawColor(const FColor& Color);
 
 		/**
 		 * Queue drawing a canvas item to the canvas.
-		 * 
+		 *
 		 * @param Item The canvas item to draw. The item *must* stay in scope until EndCanvas().
 		 * @param CanvasPosition (optional) The position to draw the item on the canvas.
+		 *
+		 * @note Not available in BPFL: TSharedRef<FCanvasItem> is not Blueprint-compatible
 		 */
-		void DrawItem(TSharedRef<FCanvasItem> ItemPtr, const FVector2f& CanvasPosition = FVector2f::ZeroVector);
+		SLATEIM_API void DrawItem(TSharedRef<FCanvasItem> ItemPtr, const FVector2f& CanvasPosition = FVector2f::ZeroVector);
 
 		/**
 		 * Queue drawing a canvas icon (texture) to the canvas.
@@ -1051,7 +1453,7 @@ namespace SlateIM
 		 * @param CanvasPosition (optional) The position of the icon on the canvas, in pixels.
 		 * @param Scale (optional) The scale of the icon.
 		 */
-		void DrawIcon(const FCanvasIcon& Icon, const FVector2f& CanvasPosition = FVector2f::ZeroVector, const FVector2f& Scale = FVector2f::UnitVector);
+		SLATEIM_API void DrawIcon(const FCanvasIcon& Icon, const FVector2f& CanvasPosition = FVector2f::ZeroVector, const FVector2f& Scale = FVector2f::UnitVector);
 
 		/**
 		 * Queue drawing a line to the canvas.
@@ -1061,31 +1463,7 @@ namespace SlateIM
 		 * @param Thickness (optional) The thickness of the line in pixels.
 		 * @param RenderColor (optional) The color of the line.
 		 */
-		void DrawLine(const FVector2D& CanvasPositionA, const FVector2D& CanvasPositionB, float Thickness = 1.0f, FLinearColor RenderColor = FLinearColor::White);
-
-		/** 
-		 * Extra optional parameters for rendering tiles (textures and materials.)
-		 */
-		struct FTileRenderParams
-		{
-			// The top left of the area on the tile to draw, in texels.
-			FVector2D UVPosition = FVector2D::ZeroVector;
-
-			// The size of the area on the tile to draw, in texels.
-			FVector2D UVSize = FVector2D::UnitVector;
-
-			// A tint to apply to tile.
-			FLinearColor RenderColor = FLinearColor::White;
-
-			// The blend mode with which to draw the tile.
-			EBlendMode BlendMode = BLEND_Translucent;
-
-			// The rotation of the tile, in degrees.
-			float Rotation = 0.f;
-
-			// The point on the tile around which it is rotated, in texels.
-			FVector2D PivotPoint = FVector2D(0.5, 0.5);
-		};
+		SLATEIM_API void DrawLine(const FVector2D& CanvasPositionA, const FVector2D& CanvasPositionB, float Thickness = 1.0f, const FLinearColor& RenderColor = FLinearColor::White);
 
 		/**
 		 * Queue drawing a texture to the canvas.
@@ -1093,9 +1471,9 @@ namespace SlateIM
 		 * @param RenderTexture The texture to draw.
 		 * @param CanvasPosition The position of the texture on the canvas, in pixels.
 		 * @param CanvasSize The size of the texture on the canvas, in pixels.
-		 * @param RenderParams (optional) Extra optional parameters for tile drawing.
+		 * @param RenderParams (optional) Additional parameters to draw the tile.
 		 */
-		void DrawTexture(UTexture* RenderTexture, const FVector2D& CanvasPosition, const FVector2D& CanvasSize, const FTileRenderParams& RenderParams = FTileRenderParams());
+		SLATEIM_API void DrawTexture(UTexture* RenderTexture, const FVector2D& CanvasPosition, const FVector2D& CanvasSize, const FTileRenderParams& RenderParams = FTileRenderParams());
 
 		/**
 		 * Queue drawing a material to the canvas.
@@ -1103,42 +1481,9 @@ namespace SlateIM
 		 * @param RenderMaterial The material to draw.
 		 * @param CanvasPosition The position of the material on the canvas, in pixels.
 		 * @param CanvasSize The size of the material on the canvas, in pixels.
-		 * @param RenderParams (optional) Extra optional parameters for tile drawing. RenderColor and BlendMode are not used.
+		 * @param RenderParams (optional) Additional parameters to draw the tile. RenderColor and BlendMode are not used.
 		 */
-		void DrawMaterial(UMaterialInterface* RenderMaterial, const FVector2D& CanvasPosition, const FVector2D& CanvasSize, const FTileRenderParams& RenderParams = FTileRenderParams());
-
-		/**
-		 * Extra optional parameters for rendering text.
-		 */
-		struct FTextRenderParams
-		{
-			// The scale of the text on the canvas.
-			FVector2D Scale = FVector2D(1.0f, 1.0f);
-
-			// The color of the text.
-			FLinearColor RenderColor = FLinearColor::White;
-			
-			// Adjust the distance between letters.
-			float Kerning = 0.0f;
-
-			// The color of the text shadow.
-			FLinearColor ShadowColor = FLinearColor::Black;
-
-			// The offset of the shadow from the text it is shadowing.
-			FVector2D ShadowOffset = FVector2D::UnitVector;
-
-			// If true, the text is centered on the given draw position on the X axis.
-			bool bCentreX = false;
-
-			// If true, the text is centered on the given draw position on the Y axis.
-			bool bCentreY = false;
-
-			// If true, outline the text.
-			bool bOutlined = false;
-
-			// Color of the text outline.
-			FLinearColor OutlineColor = FLinearColor::Black;
-		};
+		SLATEIM_API void DrawMaterial(UMaterialInterface* RenderMaterial, const FVector2D& CanvasPosition, const FVector2D& CanvasSize, const FTileRenderParams& RenderParams = FTileRenderParams());
 
 		/**
 		 * Queue drawing text to the canvas.
@@ -1146,9 +1491,9 @@ namespace SlateIM
 		 * @param RenderFont The font used to render the text.
 		 * @param RenderText The string to render.
 		 * @param CanvasPosition The position of the text on the canvas, in pixels.
-		 * @param RenderParams Extra optional parameters for rendering the text.
+		 * @param RenderParams (optional) Additional parameters to draw the text.
 		 */
-		void DrawText(UFont* RenderFont, const FString& RenderText, const FVector2D& CanvasPosition, const FTextRenderParams& RenderParams);
+		SLATEIM_API void DrawText(UFont* RenderFont, const FString& RenderText, const FVector2D& CanvasPosition, const FTextRenderParams& RenderParams);
 
 		/**
 		 * Queue drawing text to the canvas.
@@ -1161,7 +1506,7 @@ namespace SlateIM
 		 * @param Scale (optional) The scale of the text on the canvas.
 		 * @param RenderInfo (optional) Extra optional parameters for rendering the text.
 		 */
-		void DrawText(UFont* RenderFont, const FString& RenderText, const FVector2f& CanvasPosition, const FVector2f& Scale = FVector2f::UnitVector, const FFontRenderInfo& RenderInfo = FFontRenderInfo());
+		SLATEIM_API void DrawText(UFont* RenderFont, const FString& RenderText, const FVector2f& CanvasPosition, const FVector2f& Scale = FVector2f::UnitVector, const FFontRenderInfo& RenderInfo = FFontRenderInfo());
 
 		/**
 		 * Queue drawing wrapped text to canvas.
@@ -1179,24 +1524,7 @@ namespace SlateIM
 		 * @param Scale (optional) The scale of the text on the canvas.
 		 * @param RenderInfo (optional) Extra optional parameters for rendering the text.
 		 */
-		void DrawWrappedText(const UFont* RenderFont, const FString& RenderText, const FVector2f& CanvasPosition, bool bCenterTextX, bool bCenterTextY, const FVector2f& Scale = FVector2f::UnitVector, const FFontRenderInfo& RenderInfo = FFontRenderInfo());
-
-		/**
-		 * Extra optional parameters for rendering borders.
-		 * 
-		 * Tile refers to the centre texture.
-		 */
-		struct FBorderRenderParams : FTileRenderParams
-		{
-			// The scale of the border.
-			FVector2D BorderScale = FVector2D(0.1, 0.1);
-
-			// The scale of the background.
-			FVector2D BackgroundScale = FVector2D(0.1, 0.1);
-
-			// Frame corner size, in percent (should be < 0.5f).
-			FVector2D CornerSize = FVector2D::ZeroVector;
-		};
+		SLATEIM_API void DrawWrappedText(const UFont* RenderFont, const FString& RenderText, const FVector2f& CanvasPosition, bool bCenterTextX, bool bCenterTextY, const FVector2f& Scale = FVector2f::UnitVector, const FFontRenderInfo& RenderInfo = FFontRenderInfo());
 
 		/**
 		 * Queue drawing a border to the canvas.
@@ -1209,9 +1537,9 @@ namespace SlateIM
 		 * @param BottomBorderTexture The texture for the bottom of the border.
 		 * @param CanvasPosition The position of the border on the canvas, in pixels.
 		 * @param CanvasSize The size of the border on the canvas, in pixels.
-		 * @param RenderParams (optional) Extra optional params for rendering the border.
+		 * @param RenderParams (optional) Additional parameters to draw the border.
 		 */
-		void DrawBorder(UTexture* BorderTexture, UTexture* BackgroundTexture, UTexture* LeftBorderTexture, UTexture* RightBorderTexture, UTexture* TopBorderTexture, UTexture* BottomBorderTexture, const FVector2D& CanvasPosition, const FVector2D& CanvasSize, const FBorderRenderParams& RenderParams = FBorderRenderParams());
+		SLATEIM_API void DrawBorder(UTexture* BorderTexture, UTexture* BackgroundTexture, UTexture* LeftBorderTexture, UTexture* RightBorderTexture, UTexture* TopBorderTexture, UTexture* BottomBorderTexture, const FVector2D& CanvasPosition, const FVector2D& CanvasSize, const FBorderRenderParams& RenderParams = FBorderRenderParams());
 
 		/**
 		 * Queue drawing an unfilled box to the canvas.
@@ -1221,23 +1549,27 @@ namespace SlateIM
 		 * @param Thickness (optional) The thicknses of the lines to draw.
 		 * @param RenderColor (optional) The color of the lines to draw.
 		 */
-		void DrawBox(const FVector2D& CanvasPosition, const FVector2D& CanvasSize, float Thickness = 1.0f, FLinearColor RenderColor = FLinearColor::White);
+		SLATEIM_API void DrawBox(const FVector2D& CanvasPosition, const FVector2D& CanvasSize, float Thickness = 1.0f, FLinearColor RenderColor = FLinearColor::White);
 
 		/**
 		 * Queue drawing a series of texture-filled triangles to the canvas.
-		 * 
+		 *
 		 * @param RenderTexture The texture to draw with.
 		 * @param Triangles The list of triangles to draw.
+		 *
+		 * @note Not available in BPFL: TSharedRef<TArray<FCanvasUVTri>> is not Blueprint-compatible
 		 */
-		void DrawTriangles(UTexture* RenderTexture, TSharedRef<TArray<FCanvasUVTri>> Triangles);
+		SLATEIM_API void DrawTriangles(UTexture* RenderTexture, TSharedRef<TArray<FCanvasUVTri>> Triangles);
 
 		/**
 		 * Queue drawing a series of material-filled triangles to the canvas.
 		 *
 		 * @param RenderTexture The texture to draw with.
 		 * @param Triangles The list of triangles to draw.
+		 *
+		 * @note Not available in BPFL: TSharedRef<TArray<FCanvasUVTri>> is not Blueprint-compatible
 		 */
-		void DrawMaterialTriangles(UMaterialInterface* RenderMaterial, TSharedRef<TArray<FCanvasUVTri>> Triangles);
+		SLATEIM_API void DrawMaterialTriangles(UMaterialInterface* RenderMaterial, TSharedRef<TArray<FCanvasUVTri>> Triangles);
 
 		/**
 		 * Queue drawing a filled polygon to the canvas.
@@ -1248,7 +1580,7 @@ namespace SlateIM
 		 * @param NumberOfSides The number of sides of the polygon.
 		 * @param RenderColor (optional) The color of the polygon.
 		 */
-		void DrawPolygon(UTexture* RenderTexture, const FVector2D& CanvasPosition, const FVector2D& Radius, int32 NumberOfSides, FLinearColor RenderColor = FLinearColor::White);
+		SLATEIM_API void DrawPolygon(UTexture* RenderTexture, const FVector2D& CanvasPosition, const FVector2D& Radius, int32 NumberOfSides, FLinearColor RenderColor = FLinearColor::White);
 	}
 #endif // WITH_ENGINE
 #pragma endregion Canvas
diff --git a/Engine/Plugins/Experimental/SlateIM/Source/SlateIM/Public/SlateIMParameters.h b/Engine/Plugins/Experimental/SlateIM/Source/SlateIM/Public/SlateIMParameters.h
new file mode 100644
--- /dev/null
+++ b/Engine/Plugins/Experimental/SlateIM/Source/SlateIM/Public/SlateIMParameters.h
@@ -0,0 +1,752 @@
+// Copyright Epic Games, Inc. All Rights Reserved.
+
+#pragma once
+
+#include "SlateIMParametersFwd.h"
+#include "Containers/Array.h"
+#include "Containers/ArrayView.h"
+#include "Containers/UnrealString.h"
+#include "Framework/Views/ITypedTableView.h"
+#include "HAL/Platform.h"
+#include "HAL/PlatformMisc.h"
+#include "Math/Vector2D.h"
+#include "Misc/Optional.h"
+#include "Styling/SlateColor.h"
+#include "Styling/StyleColors.h"
+#include "Textures/SlateIcon.h"
+#include "Widgets/Layout/Anchors.h"
+
+#if WITH_ENGINE
+#include "SlateIMEngineParameters.h"
+#endif
+
+#include "SlateIMParameters.generated.h"
+
+namespace SlateIM
+{
+	// Sync with FSlateIMWindowParams
+	struct FWindowParams
+	{
+		// The name of the window. This will also be displayed in the windows title bar
+		FStringView WindowTitle;
+
+		// The size to open the window at (does not update the size of an existing window)
+		FVector2f WindowSize = FVector2f::ZeroVector;
+
+		// If this window was created with a previous call to BeginWindow and then closed, passing in true will reopen the window and false will leave it closed. If the window has never been seen this parameter does nothing and the window will open
+		bool bShouldReopen = false;
+
+		// Window will be create always on top.
+		bool bAlwaysOnTop = false;
+	};
+
+	// Sync with FSlateIMTableParams
+	struct FTableParams
+	{
+		// A Table View Style to override the default style. Changes to the style after the first invocation are not reflected,
+		// provide a new style object to update the visuals at runtime
+		const FTableViewStyle* Style = nullptr;
+
+		// A Table Row Style to override the default table row style. Changing the style only affects new table rows.
+		const FTableRowStyle* RowStyle = nullptr;
+
+		// Select mode for the table.
+		ESelectionMode::Type SelectionMode = ESelectionMode::None;
+	};
+
+	// Sync with FSlateIMTableColumnParams
+	struct FTableColumnParams
+	{
+		/** Added a default label to the column */
+		FStringView Label = FStringView();
+	};
+
+	// Sync with FSlateIMTableRowChildrenParams
+	struct FTableRowChildrenParams
+	{
+		/** Id of the parent row. 0 to disable state saving. Row ids must be unique to its parent. */
+		uint32 ParentRowId = 0;
+
+		/** Set the default expansion state of the parent (whether the children following are visible.) Only applies to rows with a parent id. */
+		bool bDefaultExpanded = false;
+	};
+
+	// Sync with FSlateIMTextParams
+	struct FTextParams
+	{
+		/** The color to display the text as */
+		FSlateColor Color = FSlateColor::UseForeground();
+
+		/** A Text Block Style to override the default style. Changes to the style after the first invocation are not reflected,
+		 * provide a new style object to update the visuals at runtime */
+		const FTextBlockStyle* Style = nullptr;
+	};
+
+	// Sync with FSlateIMEditableTextParams
+	struct FEditableTextParams
+	{
+		/** HintText The hint text to display when the input field is empty */
+		const FStringView HintText = FStringView();
+
+		/** An Editable Text Box Style to override the default style. Changes to the style after the first invocation are not reflected,
+		 * provide a new style object to update the visuals at runtime */
+		const FEditableTextBoxStyle* Style = nullptr;
+	};
+
+	// Sync with FSlateIMButtonParams
+	struct FButtonParams
+	{
+		/** Indicate if the button is enabled or not (grayed) */
+		bool bEnabled = true;
+
+		/** A Button Style to override the default style. Changes to the style after the first invocation are not reflected,
+		 * provide a new style object to update the visuals at runtime */
+		const FButtonStyle* Style = nullptr;
+	};
+
+	// Sync with FSlateIMCheckBoxParams
+	struct FCheckBoxParams
+	{
+		/** The label for the checkbox */
+		const FStringView Label;
+
+		/** A Checkbox Style to override the default style. Changes to the style after the first invocation are not reflected,
+		 * provide a new style object to update the visuals at runtime */
+		const FCheckBoxStyle* CheckBoxStyle = nullptr;
+	};
+
+	template<typename InNumericType>
+	struct FSpinBoxParams
+	{
+		/** The minimum value of the spin box (or unset for no limit). */
+		TOptional<InNumericType> Min;
+
+		/** The maximum value of the spin box (or unset for no limit). */
+		TOptional<InNumericType> Max;
+
+		/** A Spin Box Style to override the default style. Changes to the style after the first invocation are not reflected,
+		 * provide a new style object to update the visuals at runtime */
+		const FSpinBoxStyle* Style = nullptr;
+	};
+
+	// Sync with FSlateIMSliderParams
+	struct FSliderParams
+	{
+		/** The minimum value of the slider */
+		float Min = 0.f;
+
+		/** The maximum value of the slider */
+		float Max = 100.f;
+
+		/** The smallest incremental change that can be made to the slider's value */
+		float Step = 1.f;
+
+		/** A Slider Style to override the default style. Changes to the style after the first invocation are not reflected,
+		 * provide a new style object to update the visuals at runtime */
+		const FSliderStyle* Style = nullptr;
+	};
+
+	// Sync with FSlateIMProgressBarParams
+	struct FProgressBarParams
+	{
+		/** A Progress Bar Style to override the default style. Changes to the style after the first invocation are not reflected,
+		 * provide a new style object to update the visuals at runtime */
+		const FProgressBarStyle* Style = nullptr;
+	};
+
+	// Sync with FSlateIMComboBoxParams
+	struct FComboBoxParams
+	{
+		/** Whether to force a refresh of the available options or the selected option (set to true for a frame when changing
+		 * the list of options or manually setting the selected index) */
+		bool bForceRefresh = false;
+
+		/** Adds a search bar to the combo box */
+		bool bSearchable = false;
+
+		/** A Combo Box Style to override the default style. Changes to the style after the first invocation are not reflected,
+		 * provide a new style object to update the visuals at runtime */
+		const FComboBoxStyle* Style = nullptr;
+	};
+
+	// Sync with FSlateIMSelectionListParams
+	struct FSelectionListParams
+	{
+		/** Whether to force a refresh of the available options or the selected option (set to true for a frame when changing
+		 * the list of options or manually setting the selected index) */
+		bool bForceRefresh = false;
+
+		/** A Table View Style to override the default style. Changes to the style after the first invocation are not reflected,
+		 * provide a new style object to update the visuals at runtime */
+		const FTableViewStyle* Style = nullptr;
+	};
+
+	// Sync with FSlateIMMenuButtonParams
+	struct FMenuButtonParams
+	{
+		/** The tooltip to display for the menu item */
+		FStringView ToolTipText = FStringView();
+	};
+
+	// Sync with FSlateIMTabParams
+	struct FTabParams
+	{
+		/** An icon to display on the tab next to the title. This cannot be updated after the initial creation of the tab. */
+		const FSlateIcon TabIcon = FSlateIcon();
+
+		/** The title to display for the tab if it should differ from the TabId. This cannot be updated after the initial creation of the tab. */
+		const FText TabTitle;
+	};
+
+	// Sync with FSlateIMModalDialogParams
+	struct FModalDialogParams
+	{
+		/** The type of dialog to display */
+		EAppMsgCategory Category = EAppMsgCategory::Warning;
+
+		/** The title to display in the dialog window */
+		const FStringView Title = FStringView();
+	};
+} // SlateIM
+
+USTRUCT(BlueprintType)
+struct FSlateIMViewportRootLayout
+{
+	GENERATED_BODY()
+
+	/** How to anchor the root within the viewport (normalized units) */
+	UPROPERTY(BlueprintReadWrite, EditAnywhere, Category = "SlateIM")
+	FAnchors Anchors = FAnchors(0);
+
+	/** Offset the root from the anchored position in the viewport (in slate units) */
+	UPROPERTY(BlueprintReadWrite, EditAnywhere, Category = "SlateIM")
+	FVector2f Offset = FVector2f::ZeroVector;
+
+	/** How the root is aligned to its anchor (normalized units) */
+	UPROPERTY(BlueprintReadWrite, EditAnywhere, Category = "SlateIM")
+	FVector2f Alignment = FVector2f::ZeroVector;
+
+	/** Optional set the size (in slate units) of the root in the viewport. When unset, the root will auto-size for its content. */
+	UPROPERTY(BlueprintReadWrite, EditAnywhere, Category = "SlateIM")
+	TOptional<FVector2f> Size;
+
+	/** The ZOrder of the SlateIM widget in the viewport */
+	UPROPERTY(BlueprintReadWrite, EditAnywhere, Category = "SlateIM")
+	int32 ZOrder = 10000;
+
+	/** Global scale to apply within the viewport */
+	UPROPERTY(BlueprintReadWrite, EditAnywhere, Category = "SlateIM")
+	float Scale = 1.f;
+};
+
+USTRUCT(BlueprintType)
+struct FSlateIMWindowParams
+{
+	GENERATED_BODY()
+
+	// The name of the window. This will also be displayed in the windows title bar
+	UPROPERTY(BlueprintReadWrite, EditAnywhere, Category = "SlateIM")
+	FString WindowTitle;
+
+	// The size to open the window at (does not update the size of an existing window)
+	UPROPERTY(BlueprintReadWrite, EditAnywhere, Category = "SlateIM")
+	FVector2f WindowSize = FVector2f::ZeroVector;
+
+	// If this window was created with a previous call to BeginWindow and then closed, passing in true will reopen the window and false will leave it closed. If the window has never been seen this parameter does nothing and the window will open
+	UPROPERTY(BlueprintReadWrite, EditAnywhere, Category = "SlateIM")
+	bool bShouldReopen = false;
+
+	// Window will be create always on top.
+	UPROPERTY(BlueprintReadWrite, EditAnywhere, Category = "SlateIM")
+	bool bAlwaysOnTop = false;
+
+	operator SlateIM::FWindowParams() const
+	{
+		return {
+			WindowTitle,
+			WindowSize,
+			bShouldReopen,
+			bAlwaysOnTop
+		};
+	}
+};
+
+USTRUCT(BlueprintType)
+struct FSlateIMViewportParams
+{
+	GENERATED_BODY()
+
+	/** How to lay out the root within the viewport */
+	UPROPERTY(BlueprintReadWrite, EditAnywhere, Category = "SlateIM")
+	FSlateIMViewportRootLayout Layout;
+};
+
+USTRUCT(BlueprintType)
+struct FSlateIMTableParams
+{
+	GENERATED_BODY()
+
+	/** Select mode for the table. */
+	UPROPERTY(BlueprintReadWrite, EditAnywhere, Category = "SlateIM")
+	TEnumAsByte<ESelectionMode::Type> SelectionMode = ESelectionMode::None;
+
+	operator SlateIM::FTableParams() const
+	{
+		return {
+			nullptr,
+			nullptr,
+			SelectionMode.GetValue()
+		};
+	}
+};
+
+USTRUCT(BlueprintType)
+struct FSlateIMTableColumnParams
+{
+	GENERATED_BODY()
+
+	/** Added a default label to the column. */
+	UPROPERTY(BlueprintReadWrite, EditAnywhere, Category = "SlateIM")
+	FString Label;
+
+	operator SlateIM::FTableColumnParams() const
+	{
+		return { Label };
+	}
+};
+
+USTRUCT(BlueprintType)
+struct FSlateIMTableRowChildrenParams
+{
+	GENERATED_BODY()
+
+	/** Id of the parent row. 0 to disable state saving. Row ids must be unique to its parent. */
+	UPROPERTY(BlueprintReadWrite, EditAnywhere, Category = "SlateIM")
+	int32 ParentRowId = 0;
+
+	/** Set the default expansion state of the parent (whether the children following are visible.) Only applies to rows with a parent id. */
+	UPROPERTY(BlueprintReadWrite, EditAnywhere, Category = "SlateIM")
+	bool bDefaultExpanded = false;
+
+	operator SlateIM::FTableRowChildrenParams() const
+	{
+		return {
+			static_cast<uint32>(ParentRowId),
+			bDefaultExpanded
+		};
+	}
+};
+
+USTRUCT(BlueprintType)
+struct FSlateIMBorderParams
+{
+	GENERATED_BODY()
+
+	/** Which direction the contents of the container should flow */
+	UPROPERTY(BlueprintReadWrite, EditAnywhere, Category = "SlateIM")
+	TEnumAsByte<EOrientation> Orientation = EOrientation::Orient_Vertical;
+
+	/** Whether the container should handle all mouse inputs */
+	UPROPERTY(BlueprintReadWrite, EditAnywhere, Category = "SlateIM")
+	bool bAbsorbMouse = true;
+
+	/** How much to pad the contents of the container */
+	UPROPERTY(BlueprintReadWrite, EditAnywhere, Category = "SlateIM")
+	FMargin ContentPadding = FMargin(2.0f);
+};
+
+USTRUCT(BlueprintType)
+struct FSlateIMScrollBoxParams
+{
+	GENERATED_BODY()
+
+	/** Which direction content should flow and scroll */
+	UPROPERTY(BlueprintReadWrite, EditAnywhere, Category = "SlateIM")
+	TEnumAsByte<EOrientation> Orientation = EOrientation::Orient_Vertical;
+};
+
+USTRUCT(BlueprintType)
+struct FSlateIMPopUpParams
+{
+	GENERATED_BODY()
+
+	/** Which direction the contents of the container should flow */
+	UPROPERTY(BlueprintReadWrite, EditAnywhere, Category = "SlateIM")
+	TEnumAsByte<EOrientation> Orientation = EOrientation::Orient_Vertical;
+
+	/** Whether the container should handle all mouse inputs */
+	UPROPERTY(BlueprintReadWrite, EditAnywhere, Category = "SlateIM")
+	bool bAbsorbMouse = true;
+
+	/** How much to pad the contents of the container */
+	UPROPERTY(BlueprintReadWrite, EditAnywhere, Category = "SlateIM")
+	FMargin ContentPadding = FMargin(2.0f);
+};
+
+USTRUCT(BlueprintType)
+struct FSlateIMTextParams
+{
+	GENERATED_BODY()
+
+	/** The color to display the text as */
+	UPROPERTY(BlueprintReadWrite, EditAnywhere, Category = "SlateIM")
+	FSlateColor Color = FSlateColor::UseForeground();
+
+	operator SlateIM::FTextParams() const
+	{
+		return {
+			Color,
+			nullptr
+		};
+	}
+};
+
+USTRUCT(BlueprintType)
+struct FSlateIMEditableTextParams
+{
+	GENERATED_BODY()
+
+	/** HintText The hint text to display when the input field is empty */
+	UPROPERTY(BlueprintReadWrite, EditAnywhere, Category = "SlateIM")
+	FString HintText;
+
+	operator SlateIM::FEditableTextParams() const
+	{
+		return {
+			HintText,
+			nullptr
+		};
+	}
+};
+
+USTRUCT(BlueprintType)
+struct FSlateIMImageParams
+{
+	GENERATED_BODY()
+
+	/** The color to tint the brush */
+	UPROPERTY(BlueprintReadWrite, EditAnywhere, Category = "SlateIM")
+	FSlateColor ColorAndOpacity = FLinearColor::White;
+
+	/** Override the desired size to display the brush at */
+	UPROPERTY(BlueprintReadWrite, EditAnywhere, Category = "SlateIM")
+	FVector2D DesiredSize = FVector2D::ZeroVector;
+};
+
+USTRUCT(BlueprintType)
+struct FSlateIMButtonParams
+{
+	GENERATED_BODY()
+
+	/** Indicate if the button is enabled or not (grayed) */
+	UPROPERTY(BlueprintReadWrite, EditAnywhere, Category = "SlateIM")
+	bool bEnabled = true;
+
+	operator SlateIM::FButtonParams() const
+	{
+		return { 
+			bEnabled,
+			nullptr 
+		};
+	}
+};
+
+USTRUCT(BlueprintType)
+struct FSlateIMCheckBoxParams
+{
+	GENERATED_BODY()
+
+	/** The label for the checkbox */
+	UPROPERTY(BlueprintReadWrite, EditAnywhere, Category = "SlateIM")
+	FString Label;
+
+	operator SlateIM::FCheckBoxParams() const
+	{
+		return { 
+			Label,
+			nullptr 
+		};
+	}
+};
+
+USTRUCT(BlueprintType)
+struct FSlateIMSpinBoxFloatParams
+{
+	GENERATED_BODY()
+
+	/** The minimum value of the spin box(or unset for no limit). */
+	UPROPERTY(BlueprintReadWrite, EditAnywhere, Category = "SlateIM")
+	TOptional<float> Min;
+
+	/** The maximum value of the spin box(or unset for no limit). */
+	UPROPERTY(BlueprintReadWrite, EditAnywhere, Category = "SlateIM")
+	TOptional<float> Max;
+
+	operator SlateIM::FSpinBoxParams<float>() const
+	{
+		return {
+			Min,
+			Max,
+			nullptr
+		};
+	}
+};
+
+USTRUCT(BlueprintType)
+struct FSlateIMSpinBoxDoubleParams
+{
+	GENERATED_BODY()
+
+	/** The minimum value of the spin box(or unset for no limit). */
+	UPROPERTY(BlueprintReadWrite, EditAnywhere, Category = "SlateIM")
+	TOptional<double> Min;
+
+	/** The maximum value of the spin box(or unset for no limit). */
+	UPROPERTY(BlueprintReadWrite, EditAnywhere, Category = "SlateIM")
+	TOptional<double> Max;
+
+	operator SlateIM::FSpinBoxParams<double>() const
+	{
+		return {
+			Min,
+			Max,
+			nullptr
+		};
+	}
+};
+
+USTRUCT(BlueprintType)
+struct FSlateIMSpinBoxInt32Params
+{
+	GENERATED_BODY()
+
+	/** The minimum value of the spin box(or unset for no limit). */
+	UPROPERTY(BlueprintReadWrite, EditAnywhere, Category = "SlateIM")
+	TOptional<int32> Min;
+
+	/** The maximum value of the spin box(or unset for no limit). */
+	UPROPERTY(BlueprintReadWrite, EditAnywhere, Category = "SlateIM")
+	TOptional<int32> Max;
+
+	operator SlateIM::FSpinBoxParams<int32>() const
+	{
+		return {
+			Min,
+			Max,
+			nullptr
+		};
+	}
+};
+
+USTRUCT(BlueprintType)
+struct FSlateIMSliderParams
+{
+	GENERATED_BODY()
+
+	/** The minimum value of the slider */
+	UPROPERTY(BlueprintReadWrite, EditAnywhere, Category = "SlateIM")
+	float Min = 0.f;
+
+	/** The maximum value of the slider */
+	UPROPERTY(BlueprintReadWrite, EditAnywhere, Category = "SlateIM")
+	float Max = 100.f;
+
+	/** The smallest incremental change that can be made to the slider's value */
+	UPROPERTY(BlueprintReadWrite, EditAnywhere, Category = "SlateIM")
+	float Step = 1.f;
+
+	operator SlateIM::FSliderParams() const
+	{
+		return {
+			Min,
+			Max,
+			Step,
+			nullptr
+		};
+	}
+};
+
+USTRUCT(BlueprintType)
+struct FSlateIMProgressBarParams
+{
+	GENERATED_BODY()
+
+	operator SlateIM::FProgressBarParams() const
+	{
+		return { nullptr };
+	}
+};
+
+USTRUCT(BlueprintType)
+struct FSlateIMComboBoxParams
+{
+	GENERATED_BODY()
+
+	/** Whether to force a refresh of the available options or the selected option (set to true for a frame when changing
+	 * the list of options or manually setting the selected index) */
+	UPROPERTY(BlueprintReadWrite, EditAnywhere, Category = "SlateIM")
+	bool bForceRefresh = false;
+
+	/** Adds a search bar to the combo box */
+	UPROPERTY(BlueprintReadWrite, EditAnywhere, Category = "SlateIM")
+	bool bSearchable = false;
+
+	operator SlateIM::FComboBoxParams() const
+	{
+		return {
+			bForceRefresh,
+			bSearchable,
+			nullptr
+		};
+	}
+};
+
+USTRUCT(BlueprintType)
+struct FSlateIMSelectionListParams
+{
+	GENERATED_BODY()
+
+	/** Whether to force a refresh of the available options or the selected option (set to true for a frame when changing
+	 * the list of options or manually setting the selected index) */
+	UPROPERTY(BlueprintReadWrite, EditAnywhere, Category = "SlateIM")
+	bool bForceRefresh = false;
+
+	operator SlateIM::FSelectionListParams() const
+	{
+		return {
+			bForceRefresh,
+			nullptr
+		};
+	}
+};
+
+USTRUCT(BlueprintType)
+struct FSlateIMMenuButtonParams
+{
+	GENERATED_BODY()
+
+	/** The tooltip to display for the menu item */
+	UPROPERTY(BlueprintReadWrite, EditAnywhere, Category = "SlateIM")
+	FString ToolTipText;
+
+	operator SlateIM::FMenuButtonParams() const
+	{
+		return { ToolTipText };
+	}
+};
+
+USTRUCT(BlueprintType)
+struct FSlateIMIcon
+{
+	GENERATED_BODY()
+
+	/** The name of the style set the icon can be found in. */
+	UPROPERTY(BlueprintReadWrite, EditAnywhere, Category = "SlateIM")
+	FName StyleSetName = NAME_None;
+
+	/** The name of the style for the icon. */
+	UPROPERTY(BlueprintReadWrite, EditAnywhere, Category = "SlateIM")
+	FName StyleName = NAME_None;
+
+	/** The name of the style for the small icon. */
+	UPROPERTY(BlueprintReadWrite, EditAnywhere, Category = "SlateIM")
+	FName SmallStyleName = NAME_None;
+
+	/** Name of the style for the status overlay icon */
+	UPROPERTY(BlueprintReadWrite, EditAnywhere, Category = "SlateIM")
+	FName StatusOverlayStyleName = NAME_None;
+};
+
+USTRUCT(BlueprintType)
+struct FSlateIMTabParams
+{
+	GENERATED_BODY()
+
+	/** An icon to display on the tab next to the title. This cannot be updated after the initial creation of the tab. */
+	UPROPERTY(BlueprintReadWrite, EditAnywhere, Category = "SlateIM")
+	FSlateIMIcon TabIcon;
+
+	/** The title to display for the tab if it should differ from the TabId. This cannot be updated after the initial creation of the tab. */
+	UPROPERTY(BlueprintReadWrite, EditAnywhere, Category = "SlateIM")
+	FText TabTitle;
+
+	operator SlateIM::FTabParams() const
+	{
+		return {
+			FSlateIcon(TabIcon.StyleSetName, TabIcon.StyleName, TabIcon.SmallStyleName, TabIcon.StatusOverlayStyleName),
+			TabTitle
+		};
+	}
+};
+
+USTRUCT(BlueprintType)
+struct FSlateIMModalDialogParams
+{
+	GENERATED_BODY()
+
+	/** The type of dialog to display */
+	UPROPERTY(BlueprintReadWrite, EditAnywhere, Category = "SlateIM")
+	EAppMsgCategory Category = EAppMsgCategory::Warning;
+
+	/** The title to display in the dialog window */
+	UPROPERTY(BlueprintReadWrite, EditAnywhere, Category = "SlateIM")
+	FString Title;
+
+	operator SlateIM::FModalDialogParams() const
+	{
+		return {
+			Category,
+			Title
+		};
+	}
+};
+
+UENUM(BlueprintType)
+enum class ESlateIMFocusDepth : uint8
+{
+	// Only check if the previous widget is focused
+	SelfOnly,
+
+	// Check if the previous widget or any of its childs widgets are focused (recursively)
+	IncludingDescendants
+};
+
+USTRUCT(BlueprintType)
+struct FSlateIMGraphLinePointsParams
+{
+	GENERATED_BODY()
+
+	/** The min and max X values to horizontally scale the graph to */
+	UPROPERTY(BlueprintReadWrite, EditAnywhere, Category = "SlateIM")
+	FDoubleRange XViewRange = FDoubleRange(0.0, 1.0);
+
+	/** The min and max Y values to vertically scale the graph to */
+	UPROPERTY(BlueprintReadWrite, EditAnywhere, Category = "SlateIM")
+	FDoubleRange YViewRange = FDoubleRange(0.0, 1.0);
+
+	/** The color of the line to draw for this graph */
+	UPROPERTY(BlueprintReadWrite, EditAnywhere, Category = "SlateIM")
+	FLinearColor LineColor = FLinearColor::White;
+
+	/** How thick to draw the line */
+	UPROPERTY(BlueprintReadWrite, EditAnywhere, Category = "SlateIM")
+	float LineThickness = 1.f;
+};
+
+USTRUCT(BlueprintType)
+struct FSlateIMGraphLineValuesParams
+{
+	GENERATED_BODY()
+
+	/** The min and max Y values to vertically scale the graph to */
+	UPROPERTY(BlueprintReadWrite, EditAnywhere, Category = "SlateIM")
+	FDoubleRange ViewRange = FDoubleRange(0.0, 1.0);
+
+	/** The color of the line to draw for this graph */
+	UPROPERTY(BlueprintReadWrite, EditAnywhere, Category = "SlateIM")
+	FLinearColor LineColor = FLinearColor::White;
+
+	/** How thick to draw the line */
+	UPROPERTY(BlueprintReadWrite, EditAnywhere, Category = "SlateIM")
+	float LineThickness = 1.f;
+};
diff --git a/Engine/Plugins/Experimental/SlateIM/Source/SlateIM/Public/SlateIMParametersFwd.h b/Engine/Plugins/Experimental/SlateIM/Source/SlateIM/Public/SlateIMParametersFwd.h
new file mode 100644
--- /dev/null
+++ b/Engine/Plugins/Experimental/SlateIM/Source/SlateIM/Public/SlateIMParametersFwd.h
@@ -0,0 +1,94 @@
+// Copyright Epic Games, Inc. All Rights Reserved.
+
+#pragma once
+
+#include "HAL/Platform.h"
+
+struct FButtonStyle;
+struct FCheckBoxStyle;
+struct FComboBoxStyle;
+struct FEditableTextBoxStyle;
+struct FProgressBarStyle;
+struct FSliderStyle;
+struct FSpinBoxStyle;
+struct FTableRowStyle;
+struct FTableViewStyle;
+struct FTextBoxStyle;
+
+struct FSlateIMViewportRootLayout;
+struct FSlateIMWindowParams;
+struct FSlateIMViewportParams;
+struct FSlateIMTableParams;
+struct FSlateIMTableColumnParams;
+struct FSlateIMTableRowChildrenParams;
+struct FSlateIMBorderParams;
+struct FSlateIMScrollBoxParams;
+struct FSlateIMPopUpParams;
+struct FSlateIMTextParams;
+struct FSlateIMEditableTextParams;
+struct FSlateIMImageParams;
+struct FSlateIMButtonParams;
+struct FSlateIMCheckBoxParams;
+struct FSlateIMSpinBoxFloatParams;
+struct FSlateIMSpinBoxDoubleParams;
+struct FSlateIMSpinBoxInt32Params;
+struct FSlateIMSliderParams;
+struct FSlateIMProgressBarParams;
+struct FSlateIMComboBoxParams;
+struct FSlateIMSelectionListParams;
+struct FSlateIMMenuButtonParams;
+struct FSlateIMIcon;
+struct FSlateIMTabParams;
+struct FSlateIMModalDialogParams;
+enum class ESlateIMFocusDepth : uint8;
+struct FSlateIMGraphLinePointsParams;
+struct FSlateIMGraphLineValuesParams;
+
+#if WITH_ENGINE
+enum class ESlateIMEngineCanvasUpdateType : uint8;
+struct FSlateIMEngineCanvasParams;
+struct FSlateIMEngineTileRenderParams;
+struct FSlateIMEngineTextRenderParams;
+struct FSlateIMEngineBorderRenderParams;
+#endif
+
+namespace SlateIM
+{
+	using FViewportRootLayout = FSlateIMViewportRootLayout;
+	using FViewportParams = FSlateIMViewportParams;
+	using FBorderParams = FSlateIMBorderParams;
+	using FScrollBoxParams = FSlateIMScrollBoxParams;
+	using FPopUpParams = FSlateIMPopUpParams;
+	using FImageParams = FSlateIMImageParams;
+	using EFocusDepth = ESlateIMFocusDepth;
+	using FGraphLinePointsParams = FSlateIMGraphLinePointsParams;
+	using FGraphLineValuesParams = FSlateIMGraphLineValuesParams;
+
+#if WITH_ENGINE
+	namespace Canvas
+	{
+		using ECanvasUpdateType = ESlateIMEngineCanvasUpdateType;
+		using FCanvasParams = FSlateIMEngineCanvasParams;
+		using FTileRenderParams = FSlateIMEngineTileRenderParams;
+		using FTextRenderParams = FSlateIMEngineTextRenderParams;
+		using FBorderRenderParams = FSlateIMEngineBorderRenderParams;
+	}
+#endif
+
+	struct FWindowParams;
+	struct FTableParams;
+	struct FTableColumnParams;
+	struct FTableRowChildrenParams;
+	struct FTextParams;
+	struct FEditableTextParams;
+	struct FButtonParams;
+	struct FCheckBoxParams;
+	template<typename InNumericType> struct FSpinBoxParams;
+	struct FSliderParams;
+	struct FProgressBarParams;
+	struct FComboBoxParams;
+	struct FSelectionListParams;
+	struct FMenuButtonParams;
+	struct FTabParams;
+	struct FModalDialogParams;
+} // SlateIM
diff --git a/Engine/Plugins/Experimental/SlateIM/Source/SlateIM/SlateIM.Build.cs b/Engine/Plugins/Experimental/SlateIM/Source/SlateIM/SlateIM.Build.cs
--- a/Engine/Plugins/Experimental/SlateIM/Source/SlateIM/SlateIM.Build.cs
+++ b/Engine/Plugins/Experimental/SlateIM/Source/SlateIM/SlateIM.Build.cs
@@ -27,6 +27,8 @@ public class SlateIM : ModuleRules
 
 		if (Target.bCompileAgainstEngine)
 		{
+			PublicDependencyModuleNames.Add("SlateIMEngine");
+
 			PrivateDependencyModuleNames.AddRange(
 				new string[]
 				{
@@ -35,7 +37,7 @@ public class SlateIM : ModuleRules
 				});
 		}
 
-		if(Target.bCompileAgainstEditor)
+		if (Target.bCompileAgainstEditor)
 		{
 			PrivateDependencyModuleNames.AddRange(
 				new string[]
diff --git a/Engine/Plugins/Experimental/SlateIM/Source/SlateIMBlueprint/Private/SlateIMBlueprint.cpp b/Engine/Plugins/Experimental/SlateIM/Source/SlateIMBlueprint/Private/SlateIMBlueprint.cpp
new file mode 100644
--- /dev/null
+++ b/Engine/Plugins/Experimental/SlateIM/Source/SlateIMBlueprint/Private/SlateIMBlueprint.cpp
@@ -0,0 +1,5 @@
+// Copyright Epic Games, Inc. All Rights Reserved.
+
+#include "Modules/ModuleManager.h"
+
+IMPLEMENT_MODULE(FDefaultModuleImpl, SlateIMBlueprint)
diff --git a/Engine/Plugins/Experimental/SlateIM/Source/SlateIMBlueprint/Public/SlateIMBlueprintFunctionLibrary.h b/Engine/Plugins/Experimental/SlateIM/Source/SlateIMBlueprint/Public/SlateIMBlueprintFunctionLibrary.h
new file mode 100644
--- /dev/null
+++ b/Engine/Plugins/Experimental/SlateIM/Source/SlateIMBlueprint/Public/SlateIMBlueprintFunctionLibrary.h
@@ -0,0 +1,1313 @@
+// Copyright Epic Games, Inc. All Rights Reserved.
+
+#pragma once
+
+#include "Kismet/BlueprintFunctionLibrary.h"
+
+#include "Containers/StringFwd.h"
+#include "SlateIM.h"
+#include "SlateIMParametersFwd.h"
+
+#include "SlateIMBlueprintFunctionLibrary.generated.h"
+
+class UGameViewportClient;
+class ULocalPlayer;
+class UMaterialInterface;
+class UTexture;
+class UTexture2D;
+class UTextureRenderTarget2D;
+
+UCLASS(BlueprintType, ClassGroup = "SlateIM", DisplayName = "SlateIM")
+class USlateIMBlueprintFunctionLibrary : public UBlueprintFunctionLibrary
+{
+	GENERATED_BODY()
+
+public:
+	/**
+	 * Begins a new floating window root
+	 *
+	 * @param UniqueName a unique name to identify this window
+	 * @param Params Configuration parameters for the window.
+	 *
+	 * @return the current window state. True if open, false if closed or not updating
+	 *
+	 * @note Consider using FSlateIMWindowBase to handle creating and updating your Slate IM Window
+	 */
+	UFUNCTION(BlueprintCallable, Category = "SlateIM|Roots")
+	static bool BeginWindowRoot(FName UniqueName, const FSlateIMWindowParams& Params)
+	{
+		return SlateIM::BeginWindowRoot(UniqueName, Params);
+	}
+
+	/**
+	 * Begins a new root in the game viewport
+	 *
+	 * @param UniqueName a globally unique name to identify this root
+	 * @param ViewportClient the viewport client to add the widget to
+	 * @param Params Parameters used to create the root.
+	 *
+	 * @return Whether the root is valid and updating
+	 */
+	UFUNCTION(BlueprintCallable, Category = "SlateIM|Roots")
+	static bool BeginViewportClientRoot(FName UniqueName, UGameViewportClient* ViewportClient, const FSlateIMViewportParams& Params)
+	{
+		return SlateIM::BeginViewportRoot(UniqueName, ViewportClient, Params);
+	}
+
+	/**
+	 * Begins a new root in the provided player's viewport
+	 *
+	 * Please use the method below with the struct-parameter.
+	 * 
+	 * @param UniqueName a globally unique name to identify this root
+	 * @param LocalPlayer the player whose viewport to add the widget to
+	 * @param Layout How to lay out the root within the viewport
+	 * 
+	 * @return Whether the root is valid and updating
+	 */
+	UFUNCTION(BlueprintCallable, Category = "SlateIM|Roots")
+	static bool BeginViewportPlayerRoot(FName UniqueName, ULocalPlayer* LocalPlayer, const FSlateIMViewportParams& Params)
+	{
+		return SlateIM::BeginViewportRoot(UniqueName, LocalPlayer, Params);
+	}
+
+	/**
+	 * Ends any Root type, must always be called regardless of the result of the Begin function
+	 */
+	UFUNCTION(BlueprintCallable, Category = "SlateIM|Roots")
+	static void EndRoot()
+	{
+		SlateIM::EndRoot();
+	}
+
+	/**
+	 * Begins a horizontally stacked container. All widgets created within the Horizontal Stack container will be placed in-order left-to-right
+	 *
+	 * @param bMaximizeContent Maximize() the contents of the stack
+	 */
+	UFUNCTION(BlueprintCallable, Category = "SlateIM|Containers|Stack")
+	static void BeginHorizontalStack(bool bMaximizeContent = false)
+	{
+		SlateIM::BeginHorizontalStack(bMaximizeContent);
+	}
+
+	UFUNCTION(BlueprintCallable, Category = "SlateIM|Containers|Stack")
+	static void EndHorizontalStack()
+	{
+		SlateIM::EndHorizontalStack();
+	}
+
+	/**
+	 * Begins a vertically stacked container. All widgets created within the Vertical Stack container will be placed in-order top-to-bottom
+	 *
+	 * @param bMaximizeContent Maximize() the contents of the stack
+	 */
+	UFUNCTION(BlueprintCallable, Category = "SlateIM|Containers|Stack")
+	static void BeginVerticalStack(bool bMaximizeContent = false)
+	{
+		SlateIM::BeginVerticalStack(bMaximizeContent);
+	}
+
+	UFUNCTION(BlueprintCallable, Category = "SlateIM|Containers|Stack")
+	static void EndVerticalStack()
+	{
+		SlateIM::EndVerticalStack();
+	}
+
+	/**
+	 * Begins a horizontally wrapped container. All widgets created within the Horizontal Wrap container will be placed in-order left-to-right until the allotted width is filled,
+	 * then the content will begin a new line and so-on.
+	 *
+	 * @param bMaximizeContent Maximize() the contents of the wrap
+	 *
+	 * @note Wrap box children are always AutoSize and do not support Fill
+	 */
+	UFUNCTION(BlueprintCallable, Category = "SlateIM|Containers|Wrap")
+	static void BeginHorizontalWrap(bool bMaximizeContent = false)
+	{
+		SlateIM::BeginHorizontalWrap(bMaximizeContent);
+	}
+
+	UFUNCTION(BlueprintCallable, Category = "SlateIM|Containers|Wrap")
+	static void EndHorizontalWrap()
+	{
+		SlateIM::EndHorizontalWrap();
+	}
+
+	/**
+	 * Begins a vertically wrapped container. All widgets created within the Vertical Wrap container will be placed in-order top-to-bottom until the allotted height is filled,
+	 * then the content will begin a new column and so-on.
+	 *
+	 * @param bMaximizeContent Maximize() the contents of the wrap
+	 *
+	 * @note Wrap container children are always AutoSize and do not support Fill
+	 */
+	UFUNCTION(BlueprintCallable, Category = "SlateIM|Containers|Wrap")
+	static void BeginVerticalWrap(bool bMaximizeContent = false)
+	{
+		SlateIM::BeginVerticalWrap(bMaximizeContent);
+	}
+
+	UFUNCTION(BlueprintCallable, Category = "SlateIM|Containers|Wrap")
+	static void EndVerticalWrap()
+	{
+		SlateIM::EndVerticalWrap();
+	}
+
+	/**
+	 * Begin a Table container. The contents are automatically scrollable if the table is sized smaller than all its content.
+	 *
+	 * @param Params Settings for table construction.
+	 */
+	UFUNCTION(BlueprintCallable, Category = "SlateIM|Containers|Table")
+	static void BeginTable(const FSlateIMTableParams& Params)
+	{
+		SlateIM::BeginTable(Params);
+	}
+
+	/**
+	 * End the current Table container.
+	 */
+	UFUNCTION(BlueprintCallable, Category = "SlateIM|Containers|Table")
+	static void EndTable()
+	{
+		SlateIM::EndTable();
+	}
+
+	/**
+	 * Starts the header. Call this before adding any content to the table.
+	 */
+	UFUNCTION(BlueprintCallable, Category = "SlateIM|Containers|Table")
+	static void BeginTableHeader()
+	{
+		SlateIM::BeginTableHeader();
+	}
+
+	/**
+	 * Ends the table header.
+	 */
+	UFUNCTION(BlueprintCallable, Category = "SlateIM|Containers|Table")
+	static void EndTableHeader()
+	{
+		SlateIM::EndTableHeader();
+	}
+
+	/**
+	 * Moves to the next column. Allows fine control of the column body.
+	 *
+	 * @param ColumnID The name to use as ID for the column
+	 *
+	 * @note Table columns support SetToolTip, HAlign, and VAlign. See FixedTableColumnWidth() and InitialTableColumnWidth() for size options.
+	 */
+	UFUNCTION(BlueprintCallable, Category = "SlateIM|Containers|Table")
+	static void NextTableColumn(const FName& ColumnID)
+	{
+		SlateIM::NextTableColumn(ColumnID);
+	}
+
+	/**
+	 * Adds a column and immediately ends it with the given label as its body.
+	 *
+	 * @param ColumnID The name to use as ID for the column
+	 * @param Params The params used to create the column
+	 */
+	UFUNCTION(BlueprintCallable, Category = "SlateIM|Containers|Table")
+	static void AddTableColumn(const FName& ColumnID, const FSlateIMTableColumnParams& Params)
+	{
+		SlateIM::AddTableColumn(ColumnID, Params);
+	}
+
+	/**
+	 * Set the column to a fixed width, not resizable by users. Call this before calling AddTableColumn().
+	 *
+	 * @param Width The fixed width of the column
+	 */
+	UFUNCTION(BlueprintCallable, Category = "SlateIM|Containers|Table")
+	static void FixedTableColumnWidth(float Width)
+	{
+		SlateIM::FixedTableColumnWidth(Width);
+	}
+
+	/**
+	 * Set the column to an initial width, which can then be resized by users. Call this before calling AddTableColumn().
+	 *
+	 * @param Width The initial width of the column
+	 */
+	UFUNCTION(BlueprintCallable, Category = "SlateIM|Containers|Table")
+	static void InitialTableColumnWidth(float Width)
+	{
+		SlateIM::InitialTableColumnWidth(Width);
+	}
+
+	/**
+	 * Starts a table body.
+	 */
+	UFUNCTION(BlueprintCallable, Category = "SlateIM|Containers|Table")
+	static void BeginTableBody()
+	{
+		SlateIM::BeginTableBody();
+	}
+
+	/**
+	 * Ends a table body.
+	 */
+	UFUNCTION(BlueprintCallable, Category = "SlateIM|Containers|Table")
+	static void EndTableBody()
+	{
+		SlateIM::EndTableBody();
+	}
+
+	/**
+	 * Begin the content for the next table cell. The column and row are tracked internally based on the number of columns.
+	 *
+	 * @param bOutRowSelected If supplied, will return true if the row this cell belongs to is selected.
+	 *
+	 * @return Whether the contents of this table cell are visible. When false, the content can be skipped to save cycles.
+	 *
+	 * @note There is no "End" function counterpart to this call. This function supports the various slot functions (Padding, HAlign, VAlign, Min/Max Width/Height, etc).
+	 */
+	UFUNCTION(BlueprintCallable, Category = "SlateIM|Containers|Table")
+	static bool NextTableCell(bool& bOutRowSelected)
+	{
+		return SlateIM::NextTableCell(&bOutRowSelected);
+	}
+
+	/**
+	 * Begin adding content as a child to the current table row. NextTableCell() must still be called after this function.
+	 * If the table is a tree, a unique row id must be provided when starting a new set of child rows so that the expansion
+	 * state of the parent row is saved. Not providing a row id will cause the parent row to revert to the default expansion
+	 * state if its grandparent row is closed.
+	 *
+	 * @param Params Params used to begin the table children.
+	 *
+	 * @return Whether the parent row is expanded or not. When false, the child content can be skipped to save cycles.
+	 */
+	UFUNCTION(BlueprintCallable, Category = "SlateIM|Containers|Table")
+	static bool BeginTableRowChildren(const FSlateIMTableRowChildrenParams& Params)
+	{
+		return SlateIM::BeginTableRowChildren(Params);
+	}
+
+	/**
+	 * Stop adding child content to the parent row. This must be called no matter the result of BeginTableRowChildren()
+	 */
+	UFUNCTION(BlueprintCallable, Category = "SlateIM|Containers|Table")
+	static void EndTableRowChildren()
+	{
+		SlateIM::EndTableRowChildren();
+	}
+
+	/**
+	 * Begin a container with a background image.
+	 *
+	 * @param BorderStyleName The name of the brush in FAppStyle to use as the background for this container
+	 * @param Params Additional params to construct this border
+	 */
+	UFUNCTION(BlueprintCallable, Category = "SlateIM|Containers|Border")
+	static void BeginBorder(FName BorderStyleName, const FSlateIMBorderParams& Params)
+	{
+		SlateIM::BeginBorder(BorderStyleName, Params);
+	}
+
+	UFUNCTION(BlueprintCallable, Category = "SlateIM|Containers|Border")
+	static void EndBorder()
+	{
+		SlateIM::EndBorder();
+	}
+
+	/**
+	 * Begins a container that will allow the user to scroll when its content is larger than the allotted space.
+	 *
+	 * @param Params Parameters to create the scroll box
+	 *
+	 * @return Whether the user scroll the content this frame
+	 */
+	UFUNCTION(BlueprintCallable, Category = "SlateIM|Containers|Scroll")
+	static bool BeginScrollBox(const FSlateIMScrollBoxParams& Params)
+	{
+		return SlateIM::BeginScrollBox(Params);
+	}
+
+	UFUNCTION(BlueprintCallable, Category = "SlateIM|Containers|Scroll")
+	static void EndScrollBox()
+	{
+		SlateIM::EndScrollBox();
+	}
+
+	/**
+	 * Begin a container with a background image.
+	 *
+	 * @param BorderStyleName The name of the brush in FAppStyle to use as the background for this container
+	 * @param Params Additional params to construct this border
+	 */
+	UFUNCTION(BlueprintCallable, Category = "SlateIM|Containers|PopUp", meta=(BorderStyleName= "ToolPanel.GroupBorder"))
+	static void BeginPopUp(const FName BorderStyleName, const FSlateIMPopUpParams& Params)
+	{
+		SlateIM::BeginPopUp(BorderStyleName, Params);
+	}
+
+	UFUNCTION(BlueprintCallable, Category = "SlateIM|Containers|PopUp")
+	static void EndPopUp()
+	{
+		SlateIM::EndPopUp();
+	}
+
+	/**
+	 * Sets the padding for the next widget
+	 *
+	 * @param NextPadding The padding to apply to the next widget
+	 */
+	UFUNCTION(BlueprintCallable, Category = "SlateIM|Containers|Slots")
+	static void Padding(const FMargin NextPadding)
+	{
+		SlateIM::Padding(NextPadding);
+	}
+
+	/**
+	 * Sets the horizontal alignment for the next widget
+	 *
+	 * @param NextAlignment The alignment to apply to the next widget
+	 */
+	UFUNCTION(BlueprintCallable, Category = "SlateIM|Containers|Slots")
+	static void HAlign(EHorizontalAlignment NextAlignment)
+	{
+		SlateIM::HAlign(NextAlignment);
+	}
+
+	/**
+	 * Sets the vertical alignment for the next widget
+	 *
+	 * @param NextAlignment The alignment to apply to the next widget
+	 */
+	UFUNCTION(BlueprintCallable, Category = "SlateIM|Containers|Slots")
+	static void VAlign(EVerticalAlignment NextAlignment)
+	{
+		SlateIM::VAlign(NextAlignment);
+	}
+
+	/**
+	 * Set the next slot to AutoSize to its content
+	 */
+	UFUNCTION(BlueprintCallable, Category = "SlateIM|Containers|Slots")
+	static void AutoSize()
+	{
+		SlateIM::AutoSize();
+	}
+
+	/**
+	 * Sets the next slot to Fill the remaining space in its container
+	 *
+	 * @note The filled space is shared equally with any other slots set to Fill the same container
+	 */
+	UFUNCTION(BlueprintCallable, Category = "SlateIM|Containers|Slots")
+	static void Fill()
+	{
+		SlateIM::Fill();
+	}
+
+	/**
+	 * Set the minimum width the next slot can have.
+	 *
+	 * @param InMinWidth The minimum width for the next slot
+	 */
+	UFUNCTION(BlueprintCallable, Category = "SlateIM|Containers|Slots")
+	static void MinWidth(float MinWidth)
+	{
+		SlateIM::MinWidth(MinWidth);
+	}
+
+	/**
+	 * Set the minimum height the next slot can have.
+	 *
+	 * @param InMinHeight The minimum height for the next slot
+	 */
+	UFUNCTION(BlueprintCallable, Category = "SlateIM|Containers|Slots")
+	static void MinHeight(float MinHeight)
+	{
+		SlateIM::MinHeight(MinHeight);
+	}
+
+	/**
+	 * Set the maximum width the next slot can have.
+	 *
+	 * @param InMaxWidth The maximum width for the next slot
+	 */
+	UFUNCTION(BlueprintCallable, Category = "SlateIM|Containers|Slots")
+	static void MaxWidth(float MaxWidth)
+	{
+		SlateIM::MaxWidth(MaxWidth);
+	}
+
+	/**
+	 * Set the maximum height the next slot can have.
+	 *
+	 * @param InMaxHeight The maximum height for the next slot
+	 */
+	UFUNCTION(BlueprintCallable, Category = "SlateIM|Containers|Slots")
+	static void MaxHeight(float MaxHeight)
+	{
+		SlateIM::MaxHeight(MaxHeight);
+	}
+
+	/**
+	 * Set the alignments to Fill and removes autosize.
+	 */
+	UFUNCTION(BlueprintCallable, Category = "SlateIM|Containers|Slots")
+	static void Maximize()
+	{
+		SlateIM::Maximize();
+	}
+
+	/**
+	 * Display a string of text
+	 *
+	 * @param InText The text string to display
+	 * @param Params Additional parameters to create the text.
+	 */
+	UFUNCTION(BlueprintCallable, Category = "SlateIM|Text")
+	static void Text(const FString& InText, const FSlateIMTextParams& Params)
+	{
+		SlateIM::Text(InText, Params);
+	}
+
+	/**
+	 * Create a text input field
+	 *
+	 * @param InOutText The text that is in the input field. Can have a default value
+	 * @param Params Additional parameters to create this widget
+	 *
+	 * @return Whether the user changed or committed text this frame
+	 */
+	UFUNCTION(BlueprintCallable, Category = "SlateIM|Text")
+	static bool EditableText(UPARAM(Ref) FString& InOutText, const FSlateIMEditableTextParams& Params)
+	{
+		return SlateIM::EditableText(InOutText, Params);
+	}
+
+	/**
+	 * Display a slate style brush
+	 *
+	 * @param ImageStyleName The named of the brush to display from FAppStyle
+	 * @param Params Additional parameters to display this image
+	 */
+	UFUNCTION(BlueprintCallable, Category = "SlateIM|Image")
+	static void ImageName(const FName ImageStyleName, const FSlateIMImageParams& Params)
+	{
+		SlateIM::Image(ImageStyleName, Params);
+	}
+
+	/**
+	 * Display a colored box
+	 *
+	 * @param Params Additional parameters to display this image (including the color)
+	 */
+	UFUNCTION(BlueprintCallable, Category = "SlateIM|Image")
+	static void ImageBlock(const FSlateIMImageParams& Params)
+	{
+		SlateIM::Image(Params);
+	}
+
+	/**
+	 * Display a texture
+	 *
+	 * @param ImageTexture The texture to display
+	 * @param Params Additional parameters to display this image
+	 */
+	UFUNCTION(BlueprintCallable, Category = "SlateIM|Image")
+	static void Image(UTexture2D* ImageTexture, const FSlateIMImageParams& Params)
+	{
+		SlateIM::Image(ImageTexture, Params);
+	}
+
+	/**
+	 * Display a render target texture
+	 *
+	 * @param ImageRenderTarget The render target texture to display
+	 * @param Params Additional parameters to display this image
+	 */
+	UFUNCTION(BlueprintCallable, Category = "SlateIM|Image")
+	static void ImageTexture(UTextureRenderTarget2D* ImageRenderTarget, const FSlateIMImageParams& Params)
+	{
+		SlateIM::Image(ImageRenderTarget, Params);
+	}
+
+	/**
+	 * Display a material
+	 *
+	 * @param ImageMaterial The material to display
+	 * @param BrushSize The size to create the internal brush with
+	 * @param Params Additional parameters to display this image
+	 */
+	UFUNCTION(BlueprintCallable, Category = "SlateIM|Image")
+	static void ImageMaterial(UMaterialInterface* ImageMaterial, FVector2D BrushSize, const FSlateIMImageParams& Params)
+	{
+		SlateIM::Image(ImageMaterial, BrushSize, Params);
+	}
+
+	/**
+	 * Display a button with text
+	 *
+	 * @param InText The text to display on the button
+	 * @param Params Additional parameters to construct this button
+	 *
+	 * @return Whether the user clicked the button this frame
+	 */
+	UFUNCTION(BlueprintCallable, Category = "SlateIM|Controls")
+	static bool Button(const FString& InText, const FSlateIMButtonParams& Params)
+	{
+		return SlateIM::Button(InText, Params);
+	}
+
+	/**
+	 * Display a two-state checkbox
+	 *
+	 * @param InOutCurrentState The current state of the checkbox. Can provide a default value.
+	 * @param Params Additional params to create the checkbox
+	 *
+	 * @return Whether the user changed the state of the checkbox this frame
+	 */
+	UFUNCTION(BlueprintCallable, Category = "SlateIM|Controls")
+	static bool CheckBox(UPARAM(Ref) bool& InOutCurrentState, const FSlateIMCheckBoxParams& Params)
+	{
+		return SlateIM::CheckBox(InOutCurrentState, Params);
+	}
+
+	/**
+	 * Display a float-based spin box
+	 *
+	 * @param InOutValue The current value of the spin box. Can provide a default value.
+	 * @param Params Additional parameters to create this spinbox
+	 * 
+	 * @return Whether the user changed or committed the spin box value this frame
+	 */
+	UFUNCTION(BlueprintCallable, Category = "SlateIM|Controls")
+	static bool SpinBoxFloat(UPARAM(Ref) float& InOutValue, const FSlateIMSpinBoxFloatParams& Params)
+	{
+		return SlateIM::SpinBox<float>(InOutValue, Params);
+	}
+
+	/**
+	 * Display a double-based spin box
+	 *
+	 * @param InOutValue The current value of the spin box. Can provide a default value.
+	 * @param Params Additional parameters to create this spinbox
+	 *
+	 * @return Whether the user changed or committed the spin box value this frame
+	 */
+	UFUNCTION(BlueprintCallable, Category = "SlateIM|Controls")
+	static bool SpinBoxDouble(UPARAM(Ref) double& InOutValue, const FSlateIMSpinBoxDoubleParams& Params)
+	{
+		return SlateIM::SpinBox<double>(InOutValue, Params);
+	}
+
+	/**
+	 * Display an int32-based spin box
+	 *
+	 * @param InOutValue The current value of the spin box. Can provide a default value.
+	 * @param Params Additional parameters to create this spinbox
+	 *
+	 * @return Whether the user changed or committed the spin box value this frame
+	 */
+	UFUNCTION(BlueprintCallable, Category = "SlateIM|Controls")
+	static bool SpinBoxInt32(UPARAM(Ref) int32& InOutValue, const FSlateIMSpinBoxInt32Params& Params)
+	{
+		return SlateIM::SpinBox<int32>(InOutValue, Params);
+	}	
+	
+	/**
+	 * Display a float-based slider
+	 *
+	 * @param InOutValue The current value of the slider. Can provide a default value.
+	 * @param Params Additional parameters to construct this slider
+	 *
+	 * @return Whether the user changed the value of the slider this frame
+	 */
+	UFUNCTION(BlueprintCallable, Category = "SlateIM|Controls")
+	static bool Slider(UPARAM(Ref) float& InOutValue, const FSlateIMSliderParams& Params)
+	{
+		return SlateIM::Slider(InOutValue, Params);
+	}
+
+	/**
+	 * Display a progress bar
+	 *
+	 * @param Percent The value of the progress bar (0.0 to 1.0)
+	 * @param Params Additional params to construct the progress bar
+	 */
+	UFUNCTION(BlueprintCallable, Category = "SlateIM|Controls")
+	static void ProgressBar(TOptional<float> Percent, const FSlateIMProgressBarParams& Params)
+	{
+		SlateIM::ProgressBar(Percent, Params);
+	}
+
+	/**
+	 * Display a dropdown of text options
+	 *
+	 * @param ComboItems The options the user can choose from
+	 * @param InOutSelectedItemIndex The index of the currently selected option
+	 * @param Params Additional parameter to construct this combo box
+	 *
+	 * @return Whether the user changed the selected option this frame
+	 */
+	UFUNCTION(BlueprintCallable, Category = "SlateIM|Controls")
+	static bool ComboBox(const TArray<FString>& ComboItems, UPARAM(Ref) int32& InOutSelectedItemIndex, const FSlateIMComboBoxParams& Params)
+	{
+		return SlateIM::ComboBox(ComboItems, InOutSelectedItemIndex, Params);
+	}
+	
+	/**
+	 * Display a list of text options
+	 *
+	 * @param ListItems The list of options the user can choose from
+	 * @param InOutSelectedItemIndex The index of the currently selected option
+	 * @param Params Additional parameters used to construct this selection list
+	 *
+	 * @return Whether the user changed the selected option this frame
+	 */
+	UFUNCTION(BlueprintCallable, Category = "SlateIM|Controls")
+	static bool SelectionList(const TArray<FString>& ListItems, UPARAM(Ref) int32& InOutSelectedItemIndex, const FSlateIMSelectionListParams& Params)
+	{
+		return SlateIM::SelectionList(ListItems, InOutSelectedItemIndex, Params);
+	}
+
+	/**
+	 * Create a block of empty space
+	 *
+	 * @param Size The size of whitespace to create
+	 */
+	UFUNCTION(BlueprintCallable, Category = "SlateIM|Misc")
+	static void Spacer(const FVector2D& Size)
+	{
+		SlateIM::Spacer(Size);
+	}
+
+	/**
+	 * Begin an area where a menu appears with a right click
+	 *
+	 * @return Whether the menu is open
+	 */
+	UFUNCTION(BlueprintCallable, Category = "SlateIM|Menu")
+	static bool BeginContextMenuAnchor()
+	{
+		return SlateIM::BeginContextMenuAnchor();
+	}
+
+	UFUNCTION(BlueprintCallable, Category = "SlateIM|Menu")
+	static void EndContextMenuAnchor()
+	{
+		SlateIM::EndContextMenuAnchor();
+	}
+
+	/**
+	* Begin a menu bar, similar to those shown at the top of windows. Can be placed anywhere.
+	*
+	* @see AddMenuSeparator AddMenuButton AddMenuToggleButton AddMenuCheckButton BeginSubMenu
+	*/
+	UFUNCTION(BlueprintCallable, Category = "SlateIM|Menu")
+	static void BeginMenuBar()
+	{
+		SlateIM::BeginMenuBar();
+	}
+
+	/**
+	 * End the menu bar.
+	 */
+	UFUNCTION(BlueprintCallable, Category = "SlateIM|Menu")
+	static void EndMenuBar()
+	{
+		SlateIM::EndMenuBar();
+	}
+
+	/**
+	 * Adds a menu to the menu entry with the given name.
+	 */
+	UFUNCTION(BlueprintCallable, Category = "SlateIM|Menu")
+	static void AddMenuBarEntry(const FString& InMenuName)
+	{
+		SlateIM::AddMenuBarEntry(InMenuName);
+	}
+
+	/**
+	 * Ends a menu bar entry
+	 */
+	UFUNCTION(BlueprintCallable, Category = "SlateIM|Menu")
+	static void EndMenuBarEntry()
+	{
+		SlateIM::EndMenuBarEntry();
+	}
+
+	/**
+	 * Add a new section to the current context menu
+	 *
+	 * @param SectionText The label of the section to add
+	 */
+	UFUNCTION(BlueprintCallable, Category = "SlateIM|Menu")
+	static void AddMenuSection(const FString& SectionText)
+	{
+		SlateIM::AddMenuSection(SectionText);
+	}
+
+	/**
+	 * Add a separator to the current context menu
+	 */
+	UFUNCTION(BlueprintCallable, Category = "SlateIM|Menu")
+	static void AddMenuSeparator()
+	{
+		SlateIM::AddMenuSeparator();
+	}
+
+	/**
+	 * Add a menu item button to the current context menu
+	 *
+	 * @param RowText Label for the menu item
+	 * @param Params Additional parameters used to construct the menu button
+	 *
+	 * @return Whether the user clicked the menu item this frame
+	 */
+	UFUNCTION(BlueprintCallable, Category = "SlateIM|Menu")
+	static bool AddMenuButton(const FString& RowText, const FSlateIMMenuButtonParams& Params)
+	{
+		return SlateIM::AddMenuButton(RowText, Params);
+	}
+
+	/**
+	 * Display a menu item button with a checkbox
+	 *
+	 * @param RowText Label for the menu item
+	 * @param InOutCurrentState The current state of the toggle
+	 * @param Params Additional parameters used to construct the menu button
+	 *
+	 * @return Whether the user toggled the checkbox this frame
+	 */
+	UFUNCTION(BlueprintCallable, Category = "SlateIM|Menu")
+	static bool AddMenuToggleButton(const FString& RowText, UPARAM(Ref) bool& InOutCurrentState, const FSlateIMMenuButtonParams& Params)
+	{
+		return SlateIM::AddMenuToggleButton(RowText, InOutCurrentState, Params);
+	}
+
+	/**
+	 * Display a menu item button with a checkmark
+	 *
+	 * @param RowText Label for the menu item
+	 * @param InOutCurrentState The current state of the checkmark
+	 * @param Params Additional parameters used to construct the menu button
+	 *
+	 * @return Whether the user clicked the menu item this frame
+	 */
+	UFUNCTION(BlueprintCallable, Category = "SlateIM|Menu")
+	static bool AddMenuCheckButton(const FString& RowText, UPARAM(Ref) bool& InOutCurrentState, const FSlateIMMenuButtonParams& Params)
+	{
+		return SlateIM::AddMenuCheckButton(RowText, InOutCurrentState, Params);
+	}
+
+	/**
+	 * Add a submenu item to the current menu
+	 *
+	 * @param SubMenuText The label for the menu item
+	 */
+	UFUNCTION(BlueprintCallable, Category = "SlateIM|Menu")
+	static void BeginSubMenu(const FString& SubMenuText)
+	{
+		SlateIM::BeginSubMenu(SubMenuText);
+	}
+
+	UFUNCTION(BlueprintCallable, Category = "SlateIM|Menu")
+	static void EndSubMenu()
+	{
+		SlateIM::EndSubMenu();
+	}
+
+	/**
+	 * Begin a tab group, can contain TabStacks and TabSplitters
+	 *
+	 * @param TabGroupId Provide a unique identifier for this tab group
+	 */
+	UFUNCTION(BlueprintCallable, Category = "SlateIM|Tabs")
+	static void BeginTabGroup(FName TabGroupId)
+	{
+		SlateIM::BeginTabGroup(TabGroupId);
+	}
+
+	/**
+	 * Ends the current Tab Group
+	 */
+	UFUNCTION(BlueprintCallable, Category = "SlateIM|Tabs")
+	static void EndTabGroup()
+	{
+		SlateIM::EndTabGroup();
+	}
+
+	/**
+	 * Begin a tab stack, can only contain Tabs
+	 */
+	UFUNCTION(BlueprintCallable, Category = "SlateIM|Tabs")
+	static void BeginTabStack()
+	{
+		SlateIM::BeginTabStack();
+	}
+
+	/**
+	 * Ends the current Tab Stack
+	 */
+	UFUNCTION(BlueprintCallable, Category = "SlateIM|Tabs")
+	static void EndTabStack()
+	{
+		SlateIM::EndTabStack();
+	}
+
+	/**
+	 * Begin a tab splitter, displays child TabSplitters and TabStacks side-by-side
+	 *
+	 * @param Orientation The direction to layout child TabSplitters and TabStacks
+	 */
+	UFUNCTION(BlueprintCallable, Category = "SlateIM|Tabs")
+	static void BeginTabSplitter(TEnumAsByte<EOrientation> Orientation)
+	{
+		SlateIM::BeginTabSplitter(Orientation);
+	}
+
+	/**
+	 * Ends the current Tab Splitter
+	 */
+	UFUNCTION(BlueprintCallable, Category = "SlateIM|Tabs")
+	static void EndTabSplitter()
+	{
+		SlateIM::EndTabSplitter();
+	}
+
+	/**
+	 * Assigns the Size Coefficient for the next child of a TabSplitter
+	 *
+	 * @param SizeCoefficient The weight to assign to the next child of the parent TabSplitter
+	 */
+	UFUNCTION(BlueprintCallable, Category = "SlateIM|Tabs")
+	static void TabSplitterSizeCoefficient(float SizeCoefficient)
+	{
+		SlateIM::TabSplitterSizeCoefficient(SizeCoefficient);
+	}
+
+	/**
+	 * Begins a Tab to contain any other content
+	 *
+	 * @param TabId Provide a tab id that is unique within the Tab Group
+	 * @param Params Additional parameters used to create this tab
+	 *
+	 * @return Returns true if the tab is active, false otherwise. Logic to draw the contents of the tab can be skipped when this is false.
+	 */
+	UFUNCTION(BlueprintCallable, Category = "SlateIM|Tabs")
+	static bool BeginTab(const FName& TabId, const FSlateIMTabParams& Params)
+	{
+		return SlateIM::BeginTab(TabId, Params);
+	}
+
+	/**
+	 * Ends the current Tab. This must be called even when BeginTab returns false.
+	 */
+	UFUNCTION(BlueprintCallable, Category = "SlateIM|Tabs")
+	static void EndTab()
+	{
+		SlateIM::EndTab();
+	}
+
+	/**
+	 * Causes a one-time activation of a tab in its parent tab well. Can be used at any point during layout.
+	 * 
+	 * To use this effectively, all tab ids should be unique. Only the first tab registered with the given id will be activated.
+	 * Can be called multiple times to activate multiple tabs in the same update.
+	 * @see ActivateInParent(ETabActivationCause).
+	 * 
+	 * @param TabId The tab id to activate.
+	 */
+	UFUNCTION(BlueprintCallable, Category = "SlateIM|Tabs")
+	static void ActivateTab(const FName& TabId)
+	{
+		SlateIM::ActivateTab(TabId);
+	}
+
+	/**
+	 * SlateIM updates are disabled in specific scenarios (like when a SlateIM::ModalDialog is open), use this function to react accordingly
+	 *
+	 * @return Whether SlateIM can update currently
+	 */
+	UFUNCTION(BlueprintCallable, Category = "SlateIM|Utils")
+	static bool CanUpdateSlateIM()
+	{
+		return SlateIM::CanUpdateSlateIM();
+	}
+
+	/**
+	 * Disables all widgets until EndDisabledState is called.
+	 *
+	 * @note It is not possible to enable a widget inside a disabled parent widget by calling EndDisabledState before a child is created inside a disabled Widget.
+	 */
+	UFUNCTION(BlueprintCallable, Category = "SlateIM|Utils")
+	static void BeginDisabledState()
+	{
+		SlateIM::BeginDisabledState();
+	}
+
+	UFUNCTION(BlueprintCallable, Category = "SlateIM|Utils")
+	static void EndDisabledState()
+	{
+		SlateIM::EndDisabledState();
+	}
+
+	/**
+	 * Sets the tooltip to be used for the next widget created. Resets after tooltip is used
+	 *
+	 * @param NextToolTip The tooltip to display
+	 */
+	UFUNCTION(BlueprintCallable, Category = "SlateIM|Utils")
+	static void SetToolTip(const FString& NextToolTip)
+	{
+		SlateIM::SetToolTip(NextToolTip);
+	}
+
+	/**
+	 * Opens a modal dialog of the specified type. This function will not return until the user closes the modal dialog
+	 *
+	 * @param MessageType The type of options the user can respond to the dialog with
+	 * @param DialogText The message to display to the user
+	 * @param Params Additional parameters used to create this dialog
+	 *
+	 * @return The option the user selected
+	 *
+	 * @note SlateIM updates are disable while a SlateIM modal is open
+	 * @see CanUpdateSlateIM()
+	 */
+	UFUNCTION(BlueprintCallable, Category = "SlateIM|Containers")
+	static EAppReturnType::Type ModalDialog(EAppMsgType::Type MessageType, const FString& DialogText, const FSlateIMModalDialogParams& Params)
+	{
+		return SlateIM::ModalDialog(MessageType, DialogText, Params);
+	}
+	
+	/**
+	 * Query whether the last widget is hovered or not
+	 * 
+	 * @return true if the cursor is hovering the last rendered widget, false otherwise
+	 */
+	UFUNCTION(BlueprintCallable, Category = "SlateIM|Queries")
+	static bool IsHovered()
+	{
+		return SlateIM::IsHovered();
+	}
+
+	/**
+	 * Query whether the previous widget is focused
+	 * 
+	 * @param Depth How far to check for focus. Includes child widgets by default.
+	 * @return true if the previous widget has focus (for the specified depth), false otherwise
+	 */
+	UFUNCTION(BlueprintCallable, Category = "SlateIM|Queries", meta=(Depth="/Script/SlateIm.ScriptESlateIMFocusDepth::IncludingDescendants"))
+	static bool IsFocused(ESlateIMFocusDepth Depth)
+	{
+		return SlateIM::IsFocused(Depth);
+	}
+
+	/**
+	 * Begin a graph widget, call graphing functions between this and EndGraph to include multiple graphs in a single chart
+	 */
+	UFUNCTION(BlueprintCallable, Category = "SlateIM|Graph")
+	static void BeginGraph()
+	{
+		SlateIM::BeginGraph();
+	}
+
+	/**
+	 * Call when finished calling graphing functions for the widget
+	 */
+	UFUNCTION(BlueprintCallable, Category = "SlateIM|Graph")
+	static void EndGraph()
+	{
+		SlateIM::EndGraph();
+	}
+
+	/**
+	 * Add a line graph of 2D vectors to the graph widget
+	 *
+	 * @param Points The X,Y points to plot on the graph
+	 * @param Params Additional parameters used to draw the line
+	 */
+	UFUNCTION(BlueprintCallable, Category = "SlateIM|Graph")
+	static void GraphLinePoints(const TArray<FVector2D>& Points, const FSlateIMGraphLinePointsParams& Params)
+	{
+		SlateIM::GraphLine(const_cast<TArray<FVector2D>&>(Points), Params);
+	}
+
+	/**
+	 * Add a line graph of values to the graph widget.
+	 * The value index is used as the X-value when plotting, the graph will scale horizontally to fit all values in the array.
+	 *
+	 * @param Values The Y values to plot on the graph
+	 * @param Params Additional parameters used to draw the line
+	 */
+	UFUNCTION(BlueprintCallable, Category = "SlateIM|Graph")
+	static void GraphLineValues(const TArray<double>& Values, const FSlateIMGraphLineValuesParams& Params)
+	{
+		SlateIM::GraphLine(const_cast<TArray<double>&>(Values), Params);
+	}
+
+	/**
+	 * Query whether a key was pressed this frame.
+	 *
+	 * Only considers input events that make it to the SlateIM root. Input events handled by child widgets will not be considered.
+	 *
+	 * @param InKey The key to query the state of
+	 * @return true if the key was just pressed this frame
+	 */
+	UFUNCTION(BlueprintCallable, Category = "SlateIM|Input")
+	static bool IsKeyPressed(const FKey& InKey)
+	{
+		return SlateIM::IsKeyPressed(InKey);
+	}
+
+	/**
+	 * Query whether a key is being held this frame.
+	 *
+	 * Only considers input events that make it to the SlateIM root. Input events handled by child widgets will not be considered.
+	 *
+	 * @param InKey The key to query the state of
+	 * @return true if the key was pressed or held this frame
+	 */
+	UFUNCTION(BlueprintCallable, Category = "SlateIM|Input")
+	static bool IsKeyHeld(const FKey& InKey)
+	{
+		return SlateIM::IsKeyHeld(InKey);
+	}
+
+	/**
+	 * Query whether a key was released this frame.
+	 *
+	 * Only considers input events that make it to the SlateIM root. Input events handled by child widgets will not be considered.
+	 *
+	 * @param InKey The key to query the state of
+	 * @return true if the key was just released this frame
+	 */
+	UFUNCTION(BlueprintCallable, Category = "SlateIM|Input")
+	static bool IsKeyReleased(const FKey& InKey)
+	{
+		return SlateIM::IsKeyReleased(InKey);
+	}
+
+	/**
+	 * Retrieve the analog value for a key
+	 *
+	 * Only considers input events that make it to the SlateIM root. Input events handled by child widgets will not be considered.
+	 *
+	 * @param InKey The key to query the state of
+	 * @return the last analog value the SlateIM root received for the specified key
+	 */
+	UFUNCTION(BlueprintCallable, Category = "SlateIM|Input")
+	static float GetKeyAnalogValue(const FKey& InKey)
+	{
+		return SlateIM::GetKeyAnalogValue(InKey);
+	}
+
+	/**
+	 * Begin queueing commands to draw to a canvas render target.
+	 * 
+	 * @param Width The width of the canvas in pixels.
+	 * @param Height The height of the canvas in pixels.
+	 * @param Params Additional parameters to create this canvas.
+	 */
+	UFUNCTION(BlueprintCallable, Category = "SlateIM|Canvas")
+	static void BeginCanvas(int32 Width, int32 Height, const FSlateIMEngineCanvasParams& CanvasParams)
+	{
+		SlateIM::Canvas::BeginCanvas(Width, Height, CanvasParams);
+	}
+
+	/**
+	 * Ends the currently drawn canvas and executes all queued commands.
+	 */
+	UFUNCTION(BlueprintCallable, Category = "SlateIM|Canvas")
+	static void EndCanvas()
+	{
+		SlateIM::Canvas::EndCanvas();
+	}
+
+	/**
+	 * Invalidate the current canvas, forcing a redraw if its UpdateType is set to Invalidation.
+	 */
+	UFUNCTION(BlueprintCallable, Category = "SlateIM|Canvas")
+	static void Invalidate()
+	{
+		SlateIM::Canvas::Invalidate();
+	}
+
+	/** Sets the position of the lower-right corner of the clipping region of the canvas */
+	UFUNCTION(BlueprintCallable, Category = "SlateIM|Canvas")
+	static void SetClip(const FVector2f& ClipPosition)
+	{
+		SlateIM::Canvas::SetClip(ClipPosition);
+	}
+
+	/** Sets the draw color of the canvas. */
+	UFUNCTION(BlueprintCallable, Category = "SlateIM|Canvas")
+	static void SetLinearDrawColor(const FLinearColor& Color)
+	{
+		SlateIM::Canvas::SetDrawColor(Color);
+	}
+
+	UFUNCTION(BlueprintCallable, Category = "SlateIM|Canvas")
+	static void SetDrawColor(const FColor& Color)
+	{
+		SlateIM::Canvas::SetDrawColor(Color);
+	}
+
+	/**
+	 * Queue drawing a canvas icon (texture) to the canvas.
+	 *
+	 * @param Icon The icon to draw. The texture draw area is in pixels.
+	 * @param CanvasPosition The position of the icon on the canvas, in pixels.
+	 * @param Scale The scale of the icon.
+	 */
+	UFUNCTION(BlueprintCallable, Category = "SlateIM|Canvas")
+	static void DrawIcon(const FCanvasIcon& Icon, const FVector2f& CanvasPosition, const FVector2f& Scale)
+	{
+		SlateIM::Canvas::DrawIcon(Icon, CanvasPosition, Scale);
+	}
+
+	/**
+	 * Queue drawing a line to the canvas.
+	 *
+	 * @param CanvasPositionA One end of the line on the canvas, in pixels.
+	 * @param CanvasPositionB The other end of the line on the canvas, in pixels.
+	 * @param Thickness The thickness of the line in pixels.
+	 * @param RenderColor The color of the line.
+	 */
+	UFUNCTION(BlueprintCallable, Category = "SlateIM|Canvas")
+	static void DrawLine(const FVector2D& CanvasPositionA, const FVector2D& CanvasPositionB, float Thickness, FLinearColor RenderColor)
+	{
+		SlateIM::Canvas::DrawLine(CanvasPositionA, CanvasPositionB, Thickness, RenderColor);
+	}
+
+	/**
+	 * Queue drawing a texture to the canvas.
+	 *
+	 * @param RenderTexture The texture to draw.
+	 * @param CanvasPosition The position of the texture on the canvas, in pixels.
+	 * @param CanvasSize The size of the texture on the canvas, in pixels.
+	 * @param RenderParams Extra optional parameters for tile drawing.
+	 */
+	UFUNCTION(BlueprintCallable, Category = "SlateIM|Canvas")
+	static void DrawTexture(UTexture* RenderTexture, const FVector2D& CanvasPosition, const FVector2D& CanvasSize, const FSlateIMEngineTileRenderParams& RenderParams)
+	{
+		SlateIM::Canvas::DrawTexture(RenderTexture, CanvasPosition, CanvasSize, RenderParams);
+	}
+
+	/**
+	 * Queue drawing a material to the canvas.
+	 *
+	 * @param RenderMaterial The material to draw.
+	 * @param CanvasPosition The position of the material on the canvas, in pixels.
+	 * @param CanvasSize The size of the material on the canvas, in pixels.
+	 * @param RenderParams Extra optional parameters for tile drawing. RenderColor and BlendMode are not used.
+	 */
+	UFUNCTION(BlueprintCallable, Category = "SlateIM|Canvas")
+	static void DrawMaterial(UMaterialInterface* RenderMaterial, const FVector2D& CanvasPosition, const FVector2D& CanvasSize, const FSlateIMEngineTileRenderParams& RenderParams)
+	{
+		SlateIM::Canvas::DrawMaterial(RenderMaterial, CanvasPosition, CanvasSize, RenderParams);
+	}
+
+	/**
+	 * Queue drawing text to the canvas.
+	 *
+	 * @param RenderFont The font used to render the text.
+	 * @param RenderText The string to render.
+	 * @param CanvasPosition The position of the text on the canvas, in pixels.
+	 * @param RenderParams Extra optional parameters for rendering the text.
+	 */
+	UFUNCTION(BlueprintCallable, Category = "SlateIM|Canvas")
+	static void DrawText(UFont* RenderFont, const FString& RenderText, const FVector2D& CanvasPosition, const FSlateIMEngineTextRenderParams& RenderParams)
+	{
+		SlateIM::Canvas::DrawText(RenderFont, RenderText, CanvasPosition, RenderParams);
+	}
+
+	/**
+	 * Queue drawing text to the canvas.
+	 *
+	 * Alternate method allowing the use of FFontRenderInfo.
+	 *
+	 * @param RenderFont The font used to render the text.
+	 * @param RenderText The string to render.
+	 * @param CanvasPosition The position of the text on the canvas, in pixels.
+	 * @param Scale The scale of the text on the canvas.
+	 * @param RenderInfo Extra optional parameters for rendering the text.
+	 */
+	UFUNCTION(BlueprintCallable, Category = "SlateIM|Canvas")
+	static void DrawTextRenderInfo(UFont* RenderFont, const FString& RenderText, const FVector2f& CanvasPosition, const FVector2f& Scale, const FFontRenderInfo& RenderInfo)
+	{
+		SlateIM::Canvas::DrawText(RenderFont, RenderText, CanvasPosition, Scale, RenderInfo);
+	}
+
+	/**
+	 * Queue drawing wrapped text to canvas.
+	 *
+	 * Uses FFontRenderInfo for text effets.
+	 *
+	 * Text size / wrapping is automatically computed based on the position of the text and the @see SetCanvasClip.
+	 * The default clipping position is the entire canvas. It does not reset between queued calls.
+	 *
+	 * @param RenderFont The font used to render the text.
+	 * @param RenderText The text to render.
+	 * @param CanvasPosition The position of the text on the canvas, in pixels.
+	 * @param bCenterTextX If true, the text is centered on the given draw position on the X axis.
+	 * @param bCenterTextY If true, the text is centered on the given draw position on the Y axis.
+	 * @param Scale The scale of the text on the canvas.
+	 * @param RenderInfo Extra optional parameters for rendering the text.
+	 */
+	UFUNCTION(BlueprintCallable, Category = "SlateIM|Canvas")
+	static void DrawWrappedText(const UFont* RenderFont, const FString& RenderText, const FVector2f& CanvasPosition, bool bCenterTextX, bool bCenterTextY, const FVector2f& Scale, const FFontRenderInfo& RenderInfo)
+	{
+		SlateIM::Canvas::DrawWrappedText(RenderFont, RenderText, CanvasPosition, bCenterTextX, bCenterTextY, Scale, RenderInfo);
+	}
+
+	/**
+	 * Queue drawing a border to the canvas.
+	 *
+	 * @param BorderTexture The texture for the center of the border.
+	 * @param BackgroundTexture The texture for the background of the border.
+	 * @param LeftBorderTexture The texture for the left of the border.
+	 * @param RightBorderTexture The texture for the right of the border.
+	 * @param TopBorderTexture The texture for the top of the border.
+	 * @param BottomBorderTexture The texture for the bottom of the border.
+	 * @param CanvasPosition The position of the border on the canvas, in pixels.
+	 * @param CanvasSize The size of the border on the canvas, in pixels.
+	 * @param RenderParams Extra optional params for rendering the border.
+	 */
+	UFUNCTION(BlueprintCallable, Category = "SlateIM|Canvas")
+	static void DrawBorder(UTexture* BorderTexture, UTexture* BackgroundTexture, UTexture* LeftBorderTexture, UTexture* RightBorderTexture, UTexture* TopBorderTexture, UTexture* BottomBorderTexture, const FVector2D& CanvasPosition, const FVector2D& CanvasSize, const FSlateIMEngineBorderRenderParams& RenderParams)
+	{
+		SlateIM::Canvas::DrawBorder(BorderTexture, BackgroundTexture, LeftBorderTexture, RightBorderTexture, TopBorderTexture, BottomBorderTexture, CanvasPosition, CanvasSize, RenderParams);
+	}
+
+
+	/**
+	 * Queue drawing an unfilled box to the canvas.
+	 *
+	 * @param CanvasPosition The position of the box on the canvas, in pixels.
+	 * @param CanvasSize The size of the box on the canvas, in pixels.
+	 * @param Thickness The thicknses of the lines to draw.
+	 * @param RenderColor The color of the lines to draw.
+	 */
+	UFUNCTION(BlueprintCallable, Category = "SlateIM|Canvas")
+	static void DrawBox(const FVector2D& CanvasPosition, const FVector2D& CanvasSize, float Thickness, FLinearColor RenderColor)
+	{
+		SlateIM::Canvas::DrawBox(CanvasPosition, CanvasSize, Thickness, RenderColor);
+	}
+
+	/**
+	 * Queue drawing a filled polygon to the canvas.
+	 *
+	 * @param RenderTexture The texture to draw with.
+	 * @param CanvasPosition The position of the center of the polygon on the canvas, in pixels.
+	 * @param Radius The radius of the polygon, in pixels.
+	 * @param NumberOfSides The number of sides of the polygon.
+	 * @param RenderColor The color of the polygon.
+	 */
+	UFUNCTION(BlueprintCallable, Category = "SlateIM|Canvas")
+	static void DrawPolygon(UTexture* RenderTexture, const FVector2D& CanvasPosition, const FVector2D& Radius, int32 NumberOfSides, FLinearColor RenderColor)
+	{
+		SlateIM::Canvas::DrawPolygon(RenderTexture, CanvasPosition, Radius, NumberOfSides, RenderColor);
+	}
+};
diff --git a/Engine/Plugins/Experimental/SlateIM/Source/SlateIMBlueprint/SlateIMBlueprint.Build.cs b/Engine/Plugins/Experimental/SlateIM/Source/SlateIMBlueprint/SlateIMBlueprint.Build.cs
new file mode 100644
--- /dev/null
+++ b/Engine/Plugins/Experimental/SlateIM/Source/SlateIMBlueprint/SlateIMBlueprint.Build.cs
@@ -0,0 +1,24 @@
+// Copyright Epic Games, Inc. All Rights Reserved.
+
+using UnrealBuildTool;
+
+public class SlateIMBlueprint : ModuleRules
+{
+	public SlateIMBlueprint(ReadOnlyTargetRules Target) : base(Target)
+	{
+		PCHUsage = ModuleRules.PCHUsageMode.UseExplicitOrSharedPCHs;
+		
+		PublicDependencyModuleNames.AddRange(
+			new string[]
+			{
+				"Core",
+				"CoreUObject",
+				"Engine",
+				"InputCore",
+				"SlateCore",
+				"SlateIM",
+				"SlateIMEngine",
+			}
+		);
+	}
+}
diff --git a/Engine/Plugins/Experimental/SlateIM/Source/SlateIMEngine/Private/Implementation/SlateIMEngineCanvasDrawCommandList.cpp b/Engine/Plugins/Experimental/SlateIM/Source/SlateIMEngine/Private/Implementation/SlateIMEngineCanvasDrawCommandList.cpp
new file mode 100644
--- /dev/null
+++ b/Engine/Plugins/Experimental/SlateIM/Source/SlateIMEngine/Private/Implementation/SlateIMEngineCanvasDrawCommandList.cpp
@@ -0,0 +1,59 @@
+// Copyright Epic Games, Inc. All Rights Reserved.
+
+#include "Implementation/SlateIMEngineCanvasDrawCommandList.h"
+
+USlateIMEngineCanvasDrawCommandList::USlateIMEngineCanvasDrawCommandList()
+	: UpdateType(ESlateIMEngineCanvasUpdateType::EveryFrame)
+	, bInvalidated(true)
+{
+}
+
+ESlateIMEngineCanvasUpdateType USlateIMEngineCanvasDrawCommandList::GetUpdateType() const
+{
+	return UpdateType;
+}
+
+void USlateIMEngineCanvasDrawCommandList::SetUpdateType(ESlateIMEngineCanvasUpdateType InUpdateType)
+{
+	UpdateType = InUpdateType;
+}
+
+void USlateIMEngineCanvasDrawCommandList::Invalidate()
+{
+	bInvalidated = true;
+}
+
+bool USlateIMEngineCanvasDrawCommandList::NeedsUpdate() const
+{
+	switch (UpdateType)
+	{
+		case ESlateIMEngineCanvasUpdateType::Invalidation:
+			return bInvalidated;
+
+		case ESlateIMEngineCanvasUpdateType::EveryFrame:
+			return true;
+
+		default:
+			ensureAlwaysMsgf(false, TEXT("Invalid update time [%d]."), static_cast<uint8>(UpdateType));
+			return false;
+	}
+}
+
+void USlateIMEngineCanvasDrawCommandList::EnqueueCommand(FCanvasDrawCommand&& InDrawCommand)
+{
+	DrawCommands.Add(Forward<FCanvasDrawCommand>(InDrawCommand));
+}
+
+void USlateIMEngineCanvasDrawCommandList::ProcessCommands(UCanvas* InCanvas, int32 InWidth, int32 InHeight)
+{
+	for (const FCanvasDrawCommand& DrawCommand : DrawCommands)
+	{
+		DrawCommand.ExecuteIfBound(InCanvas);
+	}
+
+	// Make sure all delegates are deleted
+	DrawCommands.Reset();
+
+	// Reset invalidation as we've now been drawn.
+	bInvalidated = false;
+}
diff --git a/Engine/Plugins/Experimental/SlateIM/Source/SlateIMEngine/Private/SlateIMEngineModule.cpp b/Engine/Plugins/Experimental/SlateIM/Source/SlateIMEngine/Private/SlateIMEngineModule.cpp
new file mode 100644
--- /dev/null
+++ b/Engine/Plugins/Experimental/SlateIM/Source/SlateIMEngine/Private/SlateIMEngineModule.cpp
@@ -0,0 +1,5 @@
+// Copyright Epic Games, Inc. All Rights Reserved.
+
+#include "Modules/ModuleManager.h"
+
+IMPLEMENT_MODULE(FDefaultModuleImpl, SlateIMEngine)
diff --git a/Engine/Plugins/Experimental/SlateIM/Source/SlateIMEngine/Public/Implementation/SlateIMEngineCanvasDrawCommandList.h b/Engine/Plugins/Experimental/SlateIM/Source/SlateIMEngine/Public/Implementation/SlateIMEngineCanvasDrawCommandList.h
new file mode 100644
--- /dev/null
+++ b/Engine/Plugins/Experimental/SlateIM/Source/SlateIMEngine/Public/Implementation/SlateIMEngineCanvasDrawCommandList.h
@@ -0,0 +1,42 @@
+// Copyright Epic Games, Inc. All Rights Reserved.
+
+#pragma once
+
+#include "UObject/Object.h"
+
+#include "Containers/Array.h"
+#include "Delegates/DelegateCombinations.h"
+#include "SlateIMEngineParameters.h"
+
+#include "SlateIMEngineCanvasDrawCommandList.generated.h"
+
+class UCanvas;
+
+UCLASS()
+class USlateIMEngineCanvasDrawCommandList : public UObject
+{
+	GENERATED_BODY()
+
+public:
+	DECLARE_DELEGATE_OneParam(FCanvasDrawCommand, UCanvas*)
+
+	SLATEIMENGINE_API USlateIMEngineCanvasDrawCommandList();
+
+	SLATEIMENGINE_API ESlateIMEngineCanvasUpdateType GetUpdateType() const;
+
+	SLATEIMENGINE_API void SetUpdateType(ESlateIMEngineCanvasUpdateType InUpdateType);
+
+	SLATEIMENGINE_API void Invalidate();
+
+	SLATEIMENGINE_API bool NeedsUpdate() const;
+
+	SLATEIMENGINE_API void EnqueueCommand(FCanvasDrawCommand&& InDrawCommand);
+
+	UFUNCTION()
+	void ProcessCommands(UCanvas* InCanvas, int32 InWidth, int32 InHeight);
+
+protected:
+	ESlateIMEngineCanvasUpdateType UpdateType;
+	bool bInvalidated = false;
+	TArray<FCanvasDrawCommand> DrawCommands;
+};
diff --git a/Engine/Plugins/Experimental/SlateIM/Source/SlateIMEngine/Public/SlateIMEngineParameters.h b/Engine/Plugins/Experimental/SlateIM/Source/SlateIMEngine/Public/SlateIMEngineParameters.h
new file mode 100644
--- /dev/null
+++ b/Engine/Plugins/Experimental/SlateIM/Source/SlateIMEngine/Public/SlateIMEngineParameters.h
@@ -0,0 +1,174 @@
+// Copyright Epic Games, Inc. All Rights Reserved.
+
+#pragma once
+
+#include "Engine/EngineTypes.h"
+#include "Math/MathFwd.h"
+#include "Styling/SlateColor.h"
+#include "Styling/StyleColors.h"
+#include "UObject/ObjectPtr.h"
+
+#include "SlateIMEngineParameters.generated.h"
+
+UENUM(BlueprintType)
+enum class ESlateIMEngineCanvasUpdateType : uint8
+{
+	/** Redraw the canvas every frame. */
+	EveryFrame,
+
+	/**
+	 * Only redraw when specifically invalidated.
+	 * Warning: Textures and materials may not be rendered correctly if they are not loaded or shaders are not cached.
+	 */
+	Invalidation
+
+	/** TODO: Automatically redraw when the hash of the draw queue changes. */
+	// HashChange
+};
+
+USTRUCT(BlueprintType)
+struct FSlateIMEngineCanvasParams
+{
+	GENERATED_BODY()
+
+	/** Tint applied to the resulting image. */
+	UPROPERTY(BlueprintReadWrite, EditAnywhere, Category = "SlateIMEngine")
+	FSlateColor ColorAndOpacity = FSlateColor(EStyleColor::White);
+
+	/** The size of the image drawn in the UI. This may be different from the size of the canvas. */
+	UPROPERTY(BlueprintReadWrite, EditAnywhere, Category = "SlateIMEngine")
+	FVector2D DesiredSize = FVector2D::ZeroVector;
+
+	/** The world related to the canvas being drawn, if any. */
+	UPROPERTY(BlueprintReadWrite, EditAnywhere, Category = "SlateIMEngine")
+	TObjectPtr<UObject> WorldContext = nullptr;
+
+	/** When the texture will be redrawn. */
+	UPROPERTY(BlueprintReadWrite, EditAnywhere, Category = "SlateIMEngine")
+	ESlateIMEngineCanvasUpdateType UpdateType = ESlateIMEngineCanvasUpdateType::EveryFrame;
+};
+
+/**
+ * Extra optional parameters for rendering tiles (textures and materials.)
+ */
+USTRUCT(BlueprintType)
+struct FSlateIMEngineTileRenderParams
+{
+	GENERATED_BODY()
+
+	/** The top left of the area on the tile to draw, in texels. */
+	UPROPERTY(BlueprintReadWrite, EditAnywhere, Category = "SlateIMEngine")
+	FVector2D UVPosition = FVector2D::ZeroVector;
+
+	/** The size of the area on the tile to draw, in texels. */
+	UPROPERTY(BlueprintReadWrite, EditAnywhere, Category = "SlateIMEngine")
+	FVector2D UVSize = FVector2D::UnitVector;
+
+	/** A tint to apply to tile. */
+	UPROPERTY(BlueprintReadWrite, EditAnywhere, Category = "SlateIMEngine")
+	FLinearColor RenderColor = FLinearColor::White;
+
+	/** The blend mode with which to draw the tile. */
+	UPROPERTY(BlueprintReadWrite, EditAnywhere, Category = "SlateIMEngine")
+	TEnumAsByte<EBlendMode> BlendMode = BLEND_Translucent;
+
+	/** The rotation of the tile, in degrees. */
+	UPROPERTY(BlueprintReadWrite, EditAnywhere, Category = "SlateIMEngine")
+	float Rotation = 0.f;
+
+	/** The point on the tile around which it is rotated, in texels. */
+	UPROPERTY(BlueprintReadWrite, EditAnywhere, Category = "SlateIMEngine")
+	FVector2D PivotPoint = FVector2D(0.5, 0.5);
+};
+
+/**
+ * Extra optional parameters for rendering text.
+ */
+USTRUCT(BlueprintType)
+struct FSlateIMEngineTextRenderParams
+{
+	GENERATED_BODY()
+
+	/** The scale of the text on the canvas. */
+	UPROPERTY(BlueprintReadWrite, EditAnywhere, Category = "SlateIMEngine")
+	FVector2D Scale = FVector2D(1.0f, 1.0f);
+
+	/** The color of the text. */
+	UPROPERTY(BlueprintReadWrite, EditAnywhere, Category = "SlateIMEngine")
+	FLinearColor RenderColor = FLinearColor::White;
+
+	/** Adjust the distance between letters. */
+	UPROPERTY(BlueprintReadWrite, EditAnywhere, Category = "SlateIMEngine")
+	float Kerning = 0.0f;
+
+	/** The color of the text shadow. */
+	UPROPERTY(BlueprintReadWrite, EditAnywhere, Category = "SlateIMEngine")
+	FLinearColor ShadowColor = FLinearColor::Black;
+
+	/** The offset of the shadow from the text it is shadowing. */
+	UPROPERTY(BlueprintReadWrite, EditAnywhere, Category = "SlateIMEngine")
+	FVector2D ShadowOffset = FVector2D::UnitVector;
+
+	/** If true, the text is centered on the given draw position on the X axis. */
+	UPROPERTY(BlueprintReadWrite, EditAnywhere, Category = "SlateIMEngine")
+	bool bCentreX = false;
+
+	/** If true, the text is centered on the given draw position on the Y axis. */
+	UPROPERTY(BlueprintReadWrite, EditAnywhere, Category = "SlateIMEngine")
+	bool bCentreY = false;
+
+	/** If true, outline the text. */
+	UPROPERTY(BlueprintReadWrite, EditAnywhere, Category = "SlateIMEngine")
+	bool bOutlined = false;
+
+	/** Color of the text outline. */
+	UPROPERTY(BlueprintReadWrite, EditAnywhere, Category = "SlateIMEngine")
+	FLinearColor OutlineColor = FLinearColor::Black;
+};
+
+/**
+ * Extra optional parameters for rendering borders.
+ *
+ * Tile refers to the centre texture.
+ */
+USTRUCT(BlueprintType)
+struct FSlateIMEngineBorderRenderParams
+{
+	GENERATED_BODY()
+
+	/** The top left of the area on the tile to draw, in texels. */
+	UPROPERTY(BlueprintReadWrite, EditAnywhere, Category = "SlateIMEngine")
+	FVector2D UVPosition = FVector2D::ZeroVector;
+
+	/** The size of the area on the tile to draw, in texels. */
+	UPROPERTY(BlueprintReadWrite, EditAnywhere, Category = "SlateIMEngine")
+	FVector2D UVSize = FVector2D::UnitVector;
+
+	/** A tint to apply to tile. */
+	UPROPERTY(BlueprintReadWrite, EditAnywhere, Category = "SlateIMEngine")
+	FLinearColor RenderColor = FLinearColor::White;
+
+	/** The blend mode with which to draw the tile. */
+	UPROPERTY(BlueprintReadWrite, EditAnywhere, Category = "SlateIMEngine")
+	TEnumAsByte<EBlendMode> BlendMode = BLEND_Translucent;
+
+	/** The rotation of the tile, in degrees. */
+	UPROPERTY(BlueprintReadWrite, EditAnywhere, Category = "SlateIMEngine")
+	float Rotation = 0.f;
+
+	/** The point on the tile around which it is rotated, in texels. */
+	UPROPERTY(BlueprintReadWrite, EditAnywhere, Category = "SlateIMEngine")
+	FVector2D PivotPoint = FVector2D(0.5, 0.5);
+
+	/** The scale of the border. */
+	UPROPERTY(BlueprintReadWrite, EditAnywhere, Category = "SlateIMEngine")
+	FVector2D BorderScale = FVector2D(0.1, 0.1);
+
+	/** The scale of the background. */
+	UPROPERTY(BlueprintReadWrite, EditAnywhere, Category = "SlateIMEngine")
+	FVector2D BackgroundScale = FVector2D(0.1, 0.1);
+
+	/** Frame corner size, in percent (should be < 0.5f). */
+	UPROPERTY(BlueprintReadWrite, EditAnywhere, Category = "SlateIMEngine")
+	FVector2D CornerSize = FVector2D::ZeroVector;
+};
diff --git a/Engine/Plugins/Experimental/SlateIM/Source/SlateIMEngine/SlateIMEngine.Build.cs b/Engine/Plugins/Experimental/SlateIM/Source/SlateIMEngine/SlateIMEngine.Build.cs
new file mode 100644
--- /dev/null
+++ b/Engine/Plugins/Experimental/SlateIM/Source/SlateIMEngine/SlateIMEngine.Build.cs
@@ -0,0 +1,21 @@
+// Copyright Epic Games, Inc. All Rights Reserved.
+
+using UnrealBuildTool;
+
+public class SlateIMEngine : ModuleRules
+{
+	public SlateIMEngine(ReadOnlyTargetRules Target) : base(Target)
+	{
+		PCHUsage = ModuleRules.PCHUsageMode.UseExplicitOrSharedPCHs;
+
+		PublicDependencyModuleNames.AddRange(
+			new string[]
+			{
+				"Core",
+				"CoreUObject",
+				"Engine",
+				"SlateCore",
+			}
+		);
+	}
+}
