SlateIM: Build issue 1287576

#Synced-CL 52294106

diff --git a/Engine/Plugins/Experimental/SlateIM/Source/SlateIMEngine/Public/Implementation/SlateIMEngineCanvasDrawCommandList.h b/Engine/Plugins/Experimental/SlateIM/Source/SlateIMEngine/Public/Implementation/SlateIMEngineCanvasDrawCommandList.h
--- a/Engine/Plugins/Experimental/SlateIM/Source/SlateIMEngine/Public/Implementation/SlateIMEngineCanvasDrawCommandList.h
+++ b/Engine/Plugins/Experimental/SlateIM/Source/SlateIMEngine/Public/Implementation/SlateIMEngineCanvasDrawCommandList.h
@@ -33,7 +33,7 @@ public:
 	SLATEIMENGINE_API void EnqueueCommand(FCanvasDrawCommand&& InDrawCommand);
 
 	UFUNCTION()
-	void ProcessCommands(UCanvas* InCanvas, int32 InWidth, int32 InHeight);
+	SLATEIMENGINE_API void ProcessCommands(UCanvas* InCanvas, int32 InWidth, int32 InHeight);
 
 protected:
 	ESlateIMEngineCanvasUpdateType UpdateType;
