[Backout] - CL52348642
#rnx
Original CL Desc
-----------------------------------------------------------------
SlateIM: Major changes
- Change to struct parameters to the majority of methods.
- Exported canvas functions
- Added Maximize option to Stack and Wrap calls.
- Updated examples file to use the struct params.
- Added BP Support
- Fixed compilation in non-engine configuration

#Synced-CL 52294106

diff --git a/Engine/Plugins/Experimental/SlateIM/SlateIM.uplugin b/Engine/Plugins/Experimental/SlateIM/SlateIM.uplugin
--- a/Engine/Plugins/Experimental/SlateIM/SlateIM.uplugin
+++ b/Engine/Plugins/Experimental/SlateIM/SlateIM.uplugin
@@ -21,20 +21,10 @@
 			"Type": "Runtime",
 			"LoadingPhase": "Default"
 		},
-		{
-			"Name": "SlateIMEngine",
-			"Type": "Runtime",
-			"LoadingPhase": "Default"
-		},
 		{
 			"Name": "SlateIMInGame",
 			"Type": "Runtime",
 			"LoadingPhase": "Default"
-		},
-		{
-			"Name": "SlateIMBlueprint",
-			"Type": "Runtime",
-			"LoadingPhase": "Default"
 		}
 	]
 }
\ No newline at end of file
diff --git a/Engine/Plugins/Experimental/SlateIM/Source/SlateIM/Private/Implementation/SlateIMCanvasDrawCommandList.cpp b/Engine/Plugins/Experimental/SlateIM/Source/SlateIM/Private/Implementation/SlateIMCanvasDrawCommandList.cpp
new file mode 100644
--- /dev/null
+++ b/Engine/Plugins/Experimental/SlateIM/Source/SlateIM/Private/Implementation/SlateIMCanvasDrawCommandList.cpp
@@ -0,0 +1,61 @@
+// Copyright Epic Games, Inc. All Rights Reserved.
+
+#include "Implementation/SlateIMCanvasDrawCommandList.h"
+
+#include "SlateIM.h"
+
+USlateIMCanvasDrawCommandList::USlateIMCanvasDrawCommandList()
+	: UpdateType(SlateIM::Canvas::ECanvasUpdateType::EveryFrame)
+	, bInvalidated(true)
+{
+}
+
+SlateIM::Canvas::ECanvasUpdateType USlateIMCanvasDrawCommandList::GetUpdateType() const
+{
+	return UpdateType;
+}
+
+void USlateIMCanvasDrawCommandList::SetUpdateType(SlateIM::Canvas::ECanvasUpdateType InUpdateType)
+{
+	UpdateType = InUpdateType;
+}
+
+void USlateIMCanvasDrawCommandList::Invalidate()
+{
+	bInvalidated = true;
+}
+
+bool USlateIMCanvasDrawCommandList::NeedsUpdate() const
+{
+	switch (UpdateType)
+	{
+		case SlateIM::Canvas::ECanvasUpdateType::Invalidation:
+			return bInvalidated;
+
+		case SlateIM::Canvas::ECanvasUpdateType::EveryFrame:
+			return true;
+
+		default:
+			ensureAlwaysMsgf(false, TEXT("Invalid update time [%d]."), static_cast<uint8>(UpdateType));
+			return false;
+	}
+}
+
+void USlateIMCanvasDrawCommandList::EnqueueCommand(FCanvasDrawCommand&& InDrawCommand)
+{
+	DrawCommands.Add(Forward<FCanvasDrawCommand>(InDrawCommand));
+}
+
+void USlateIMCanvasDrawCommandList::ProcessCommands(UCanvas* InCanvas, int32 InWidth, int32 InHeight)
+{
+	for (const FCanvasDrawCommand& DrawCommand : DrawCommands)
+	{
+		DrawCommand.ExecuteIfBound(InCanvas);
+	}
+
+	// Make sure all delegates are deleted
+	DrawCommands.Reset();
+
+	// Reset invalidation as we've now been drawn.
+	bInvalidated = false;
+}
diff --git a/Engine/Plugins/Experimental/SlateIM/Source/SlateIM/Private/Implementation/SlateIMCanvasDrawCommandList.h b/Engine/Plugins/Experimental/SlateIM/Source/SlateIM/Private/Implementation/SlateIMCanvasDrawCommandList.h
new file mode 100644
--- /dev/null
+++ b/Engine/Plugins/Experimental/SlateIM/Source/SlateIM/Private/Implementation/SlateIMCanvasDrawCommandList.h
@@ -0,0 +1,46 @@
+// Copyright Epic Games, Inc. All Rights Reserved.
+
+#pragma once
+
+#include "UObject/Object.h"
+
+#include "Containers/Array.h"
+#include "Delegates/DelegateCombinations.h"
+
+#include "SlateIMCanvasDrawCommandList.generated.h"
+
+class UCanvas;
+
+namespace SlateIM::Canvas
+{
+	enum class ECanvasUpdateType : uint8;
+}
+
+UCLASS()
+class USlateIMCanvasDrawCommandList : public UObject
+{
+	GENERATED_BODY()
+
+public:
+	DECLARE_DELEGATE_OneParam(FCanvasDrawCommand, UCanvas*)
+
+	USlateIMCanvasDrawCommandList();
+
+	SlateIM::Canvas::ECanvasUpdateType GetUpdateType() const;
+
+	void SetUpdateType(SlateIM::Canvas::ECanvasUpdateType InUpdateType);
+
+	void Invalidate();
+
+	bool NeedsUpdate() const;
+
+	void EnqueueCommand(FCanvasDrawCommand&& InDrawCommand);
+
+	UFUNCTION()
+	void ProcessCommands(UCanvas* InCanvas, int32 InWidth, int32 InHeight);
+
+protected:
+	SlateIM::Canvas::ECanvasUpdateType UpdateType;
+	bool bInvalidated = false;
+	TArray<FCanvasDrawCommand> DrawCommands;
+};
diff --git a/Engine/Plugins/Experimental/SlateIM/Source/SlateIM/Private/Implementation/SlateIM_Canvas.cpp b/Engine/Plugins/Experimental/SlateIM/Source/SlateIM/Private/Implementation/SlateIM_Canvas.cpp
--- a/Engine/Plugins/Experimental/SlateIM/Source/SlateIM/Private/Implementation/SlateIM_Canvas.cpp
+++ b/Engine/Plugins/Experimental/SlateIM/Source/SlateIM/Private/Implementation/SlateIM_Canvas.cpp
@@ -11,7 +11,7 @@
 #include "Engine/CanvasRenderTarget2D.h"
 #include "Engine/Texture2D.h"
 #include "Engine/TextureRenderTarget2D.h"
-#include "Implementation/SlateIMEngineCanvasDrawCommandList.h"
+#include "Implementation/SlateIMCanvasDrawCommandList.h"
 #include "Math/IntPoint.h"
 #include "Misc/SlateIMManager.h"
 #include "Misc/SlateIMSlateResources.h"
@@ -72,29 +72,29 @@ namespace SlateIM::Canvas
 		}
 
 		// Create uobject proxy for queueing draw calls
-		USlateIMEngineCanvasDrawCommandList* UpdateContainer = nullptr;
+		USlateIMCanvasDrawCommandList* UpdateContainer = nullptr;
 		TSharedPtr<FSlateIMPinnedObjectResourceDataStore> UpdateResource = CanvasImage->GetMetaData<FSlateIMPinnedObjectResourceDataStore>();
 
 		if (!UpdateResource)
 		{
-			UpdateContainer = NewObject<USlateIMEngineCanvasDrawCommandList>();
+			UpdateContainer = NewObject<USlateIMCanvasDrawCommandList>();
 			UpdateResource = MakeShared<FSlateIMPinnedObjectResourceDataStore>(FPinnedObjectResource{ TStrongObjectPtr<UObject>(UpdateContainer) });
 			CanvasImage->AddMetadata(UpdateResource.ToSharedRef());
 		}
 		else
 		{
-			UpdateContainer = CastChecked<USlateIMEngineCanvasDrawCommandList>(UpdateResource->Data.PinnedResource.Get());
+			UpdateContainer = CastChecked<USlateIMCanvasDrawCommandList>(UpdateResource->Data.PinnedResource.Get());
 		}
 
 		CanvasRenderTarget->OnCanvasRenderTargetUpdate.Clear();
-		CanvasRenderTarget->OnCanvasRenderTargetUpdate.AddDynamic(UpdateContainer, &USlateIMEngineCanvasDrawCommandList::ProcessCommands);
+		CanvasRenderTarget->OnCanvasRenderTargetUpdate.AddDynamic(UpdateContainer, &USlateIMCanvasDrawCommandList::ProcessCommands);
 
 		// Defer update call to EndCanvas
 
 		return &CanvasResource->Data.Brush;
 	}
 
-	USlateIMEngineCanvasDrawCommandList* GetCanvasCommandList()
+	USlateIMCanvasDrawCommandList* GetCanvasCommandList()
 	{
 		TSharedPtr<SImCompoundWidget> CanvasContainer = FSlateIMManager::Get().GetCurrentContainer<SImCompoundWidget>();
 
@@ -104,7 +104,7 @@ namespace SlateIM::Canvas
 			{
 				if (TSharedPtr<FSlateIMPinnedObjectResourceDataStore> CanvasResource = CanvasImage->GetMetaData<FSlateIMPinnedObjectResourceDataStore>())
 				{
-					return CastChecked<USlateIMEngineCanvasDrawCommandList>(CanvasResource->Data.PinnedResource.Get());
+					return CastChecked<USlateIMCanvasDrawCommandList>(CanvasResource->Data.PinnedResource.Get());
 				}
 			}
 		}
@@ -130,9 +130,9 @@ namespace SlateIM::Canvas
 				return GetCanvasBrush(ImageWidget, WorldContextWeak, Width, Height);
 			};
 
-		Image_Internal({.ColorAndOpacity = CanvasParams.ColorAndOpacity, .DesiredSize = CanvasParams.DesiredSize}, GetBrush);
+		Image_Internal(CanvasParams.ColorAndOpacity, CanvasParams.DesiredSize, GetBrush);
 
-		if (USlateIMEngineCanvasDrawCommandList* CommandList = GetCanvasCommandList())
+		if (USlateIMCanvasDrawCommandList* CommandList = GetCanvasCommandList())
 		{
 			CommandList->SetUpdateType(CanvasParams.UpdateType);
 		}
@@ -140,7 +140,7 @@ namespace SlateIM::Canvas
 
 	void EndCanvas()
 	{
-		if (USlateIMEngineCanvasDrawCommandList* CommandList = GetCanvasCommandList())
+		if (USlateIMCanvasDrawCommandList* CommandList = GetCanvasCommandList())
 		{
 			if (!CommandList->NeedsUpdate())
 			{
@@ -182,7 +182,7 @@ namespace SlateIM::Canvas
 
 	void Invalidate()
 	{
-		if (USlateIMEngineCanvasDrawCommandList* CommandList = GetCanvasCommandList())
+		if (USlateIMCanvasDrawCommandList* CommandList = GetCanvasCommandList())
 		{
 			CommandList->Invalidate();
 		}
@@ -190,9 +190,9 @@ namespace SlateIM::Canvas
 
 	void SetClip(const FVector2f& ClipPosition)
 	{
-		if (USlateIMEngineCanvasDrawCommandList* CommandList = GetCanvasCommandList())
+		if (USlateIMCanvasDrawCommandList* CommandList = GetCanvasCommandList())
 		{
-			CommandList->EnqueueCommand(USlateIMEngineCanvasDrawCommandList::FCanvasDrawCommand::CreateLambda(
+			CommandList->EnqueueCommand(USlateIMCanvasDrawCommandList::FCanvasDrawCommand::CreateLambda(
 				[ClipPosition](UCanvas* InCanvas)
 				{
 					InCanvas->SetClip(ClipPosition.X, ClipPosition.Y);
@@ -203,9 +203,9 @@ namespace SlateIM::Canvas
 
 	void SetDrawColor(const FLinearColor& Color)
 	{
-		if (USlateIMEngineCanvasDrawCommandList* CommandList = GetCanvasCommandList())
+		if (USlateIMCanvasDrawCommandList* CommandList = GetCanvasCommandList())
 		{
-			CommandList->EnqueueCommand(USlateIMEngineCanvasDrawCommandList::FCanvasDrawCommand::CreateLambda(
+			CommandList->EnqueueCommand(USlateIMCanvasDrawCommandList::FCanvasDrawCommand::CreateLambda(
 				[Color](UCanvas* InCanvas)
 				{
 					InCanvas->SetLinearDrawColor(Color);
@@ -216,9 +216,9 @@ namespace SlateIM::Canvas
 
 	void SetDrawColor(const FColor& Color)
 	{
-		if (USlateIMEngineCanvasDrawCommandList* CommandList = GetCanvasCommandList())
+		if (USlateIMCanvasDrawCommandList* CommandList = GetCanvasCommandList())
 		{
-			CommandList->EnqueueCommand(USlateIMEngineCanvasDrawCommandList::FCanvasDrawCommand::CreateLambda(
+			CommandList->EnqueueCommand(USlateIMCanvasDrawCommandList::FCanvasDrawCommand::CreateLambda(
 				[Color](UCanvas* InCanvas)
 				{
 					InCanvas->SetDrawColor(Color);
@@ -229,9 +229,9 @@ namespace SlateIM::Canvas
 
 	void DrawItem(TSharedRef<FCanvasItem> ItemPtr, const FVector2f& CanvasPosition)
 	{
-		if (USlateIMEngineCanvasDrawCommandList* CommandList = GetCanvasCommandList())
+		if (USlateIMCanvasDrawCommandList* CommandList = GetCanvasCommandList())
 		{
-			CommandList->EnqueueCommand(USlateIMEngineCanvasDrawCommandList::FCanvasDrawCommand::CreateLambda(
+			CommandList->EnqueueCommand(USlateIMCanvasDrawCommandList::FCanvasDrawCommand::CreateLambda(
 				[ItemPtr, CanvasPosition](UCanvas* InCanvas)
 				{
 					InCanvas->DrawItem(*ItemPtr, CanvasPosition.X, CanvasPosition.Y);
@@ -242,9 +242,9 @@ namespace SlateIM::Canvas
 
 	void DrawIcon(const FCanvasIcon& Icon, const FVector2f& CanvasPosition, const FVector2f& Scale)
 	{
-		if (USlateIMEngineCanvasDrawCommandList* CommandList = GetCanvasCommandList())
+		if (USlateIMCanvasDrawCommandList* CommandList = GetCanvasCommandList())
 		{
-			CommandList->EnqueueCommand(USlateIMEngineCanvasDrawCommandList::FCanvasDrawCommand::CreateLambda(
+			CommandList->EnqueueCommand(USlateIMCanvasDrawCommandList::FCanvasDrawCommand::CreateLambda(
 				[Icon, CanvasPosition, Scale](UCanvas* InCanvas)
 				{
 					InCanvas->DrawScaledIcon(
@@ -258,11 +258,11 @@ namespace SlateIM::Canvas
 		}
 	}
 
-	void DrawLine(const FVector2D& CanvasPositionA, const FVector2D& CanvasPositionB, float Thickness, const FLinearColor& RenderColor)
+	void DrawLine(const FVector2D& CanvasPositionA, const FVector2D& CanvasPositionB, float Thickness, FLinearColor RenderColor)
 	{
-		if (USlateIMEngineCanvasDrawCommandList* CommandList = GetCanvasCommandList())
+		if (USlateIMCanvasDrawCommandList* CommandList = GetCanvasCommandList())
 		{
-			CommandList->EnqueueCommand(USlateIMEngineCanvasDrawCommandList::FCanvasDrawCommand::CreateLambda(
+			CommandList->EnqueueCommand(USlateIMCanvasDrawCommandList::FCanvasDrawCommand::CreateLambda(
 				[CanvasPositionA, CanvasPositionB, Thickness, RenderColor](UCanvas* InCanvas)
 				{
 					InCanvas->K2_DrawLine(
@@ -278,9 +278,9 @@ namespace SlateIM::Canvas
 
 	void DrawTexture(UTexture* RenderTexture, const FVector2D& CanvasPosition, const FVector2D& CanvasSize, const FTileRenderParams& RenderParams)
 	{
-		if (USlateIMEngineCanvasDrawCommandList* CommandList = GetCanvasCommandList())
+		if (USlateIMCanvasDrawCommandList* CommandList = GetCanvasCommandList())
 		{
-			CommandList->EnqueueCommand(USlateIMEngineCanvasDrawCommandList::FCanvasDrawCommand::CreateLambda(
+			CommandList->EnqueueCommand(USlateIMCanvasDrawCommandList::FCanvasDrawCommand::CreateLambda(
 				[RenderTexture, CanvasPosition, CanvasSize, RenderParams](UCanvas* InCanvas)
 				{
 					InCanvas->K2_DrawTexture(
@@ -301,9 +301,9 @@ namespace SlateIM::Canvas
 
 	void DrawMaterial(UMaterialInterface* RenderMaterial, const FVector2D& CanvasPosition, const FVector2D& CanvasSize, const FTileRenderParams& RenderParams)
 	{
-		if (USlateIMEngineCanvasDrawCommandList* CommandList = GetCanvasCommandList())
+		if (USlateIMCanvasDrawCommandList* CommandList = GetCanvasCommandList())
 		{
-			CommandList->EnqueueCommand(USlateIMEngineCanvasDrawCommandList::FCanvasDrawCommand::CreateLambda(
+			CommandList->EnqueueCommand(USlateIMCanvasDrawCommandList::FCanvasDrawCommand::CreateLambda(
 				[RenderMaterial, CanvasPosition, CanvasSize, RenderParams](UCanvas* InCanvas)
 				{
 					InCanvas->K2_DrawMaterial(
@@ -322,9 +322,9 @@ namespace SlateIM::Canvas
 
 	void DrawText(UFont* RenderFont, const FString& RenderText, const FVector2D& CanvasPosition, const FTextRenderParams& RenderParams)
 	{
-		if (USlateIMEngineCanvasDrawCommandList* CommandList = GetCanvasCommandList())
+		if (USlateIMCanvasDrawCommandList* CommandList = GetCanvasCommandList())
 		{
-			CommandList->EnqueueCommand(USlateIMEngineCanvasDrawCommandList::FCanvasDrawCommand::CreateLambda(
+			CommandList->EnqueueCommand(USlateIMCanvasDrawCommandList::FCanvasDrawCommand::CreateLambda(
 				[RenderFont, RenderText, CanvasPosition, RenderParams](UCanvas* InCanvas)
 				{
 					InCanvas->K2_DrawText(
@@ -348,9 +348,9 @@ namespace SlateIM::Canvas
 
 	void DrawText(UFont* RenderFont, const FString& RenderText, const FVector2f& CanvasPosition, const FVector2f& Scale, const FFontRenderInfo& RenderInfo)
 	{
-		if (USlateIMEngineCanvasDrawCommandList* CommandList = GetCanvasCommandList())
+		if (USlateIMCanvasDrawCommandList* CommandList = GetCanvasCommandList())
 		{
-			CommandList->EnqueueCommand(USlateIMEngineCanvasDrawCommandList::FCanvasDrawCommand::CreateLambda(
+			CommandList->EnqueueCommand(USlateIMCanvasDrawCommandList::FCanvasDrawCommand::CreateLambda(
 				[RenderFont, RenderText, CanvasPosition, Scale, RenderInfo](UCanvas* InCanvas)
 				{
 					InCanvas->DrawText(
@@ -369,9 +369,9 @@ namespace SlateIM::Canvas
 
 	void DrawWrappedText(const UFont* RenderFont, const FString& RenderText, const FVector2f& CanvasPosition, bool bCenterTextX, bool bCenterTextY, const FVector2f& Scale, const FFontRenderInfo& RenderInfo)
 	{
-		if (USlateIMEngineCanvasDrawCommandList* CommandList = GetCanvasCommandList())
+		if (USlateIMCanvasDrawCommandList* CommandList = GetCanvasCommandList())
 		{
-			CommandList->EnqueueCommand(USlateIMEngineCanvasDrawCommandList::FCanvasDrawCommand::CreateLambda(
+			CommandList->EnqueueCommand(USlateIMCanvasDrawCommandList::FCanvasDrawCommand::CreateLambda(
 				[RenderFont, RenderText, CanvasPosition, Scale, bCenterTextX, bCenterTextY, RenderInfo](UCanvas* InCanvas)
 				{
 					int32 outXL, outYL;
@@ -397,9 +397,9 @@ namespace SlateIM::Canvas
 
 	void DrawBorder(UTexture* BorderTexture, UTexture* BackgroundTexture, UTexture* LeftBorderTexture, UTexture* RightBorderTexture, UTexture* TopBorderTexture, UTexture* BottomBorderTexture, const FVector2D& CanvasPosition, const FVector2D& CanvasSize, const FBorderRenderParams& RenderParams)
 	{
-		if (USlateIMEngineCanvasDrawCommandList* CommandList = GetCanvasCommandList())
+		if (USlateIMCanvasDrawCommandList* CommandList = GetCanvasCommandList())
 		{
-			CommandList->EnqueueCommand(USlateIMEngineCanvasDrawCommandList::FCanvasDrawCommand::CreateLambda(
+			CommandList->EnqueueCommand(USlateIMCanvasDrawCommandList::FCanvasDrawCommand::CreateLambda(
 				[BorderTexture,BackgroundTexture, LeftBorderTexture, RightBorderTexture, TopBorderTexture, BottomBorderTexture, 
 					CanvasPosition, CanvasSize, RenderParams](UCanvas* InCanvas)
 				{
@@ -428,9 +428,9 @@ namespace SlateIM::Canvas
 
 	void DrawBox(const FVector2D& CanvasPosition, const FVector2D& CanvasSize, float Thickness, FLinearColor RenderColor)
 	{
-		if (USlateIMEngineCanvasDrawCommandList* CommandList = GetCanvasCommandList())
+		if (USlateIMCanvasDrawCommandList* CommandList = GetCanvasCommandList())
 		{
-			CommandList->EnqueueCommand(USlateIMEngineCanvasDrawCommandList::FCanvasDrawCommand::CreateLambda(
+			CommandList->EnqueueCommand(USlateIMCanvasDrawCommandList::FCanvasDrawCommand::CreateLambda(
 				[CanvasPosition, CanvasSize, Thickness, RenderColor](UCanvas* InCanvas)
 				{
 					InCanvas->K2_DrawBox(
@@ -446,9 +446,9 @@ namespace SlateIM::Canvas
 
 	void DrawTriangles(UTexture* RenderTexture, TSharedRef<TArray<FCanvasUVTri>> Triangles)
 	{
-		if (USlateIMEngineCanvasDrawCommandList* CommandList = GetCanvasCommandList())
+		if (USlateIMCanvasDrawCommandList* CommandList = GetCanvasCommandList())
 		{
-			CommandList->EnqueueCommand(USlateIMEngineCanvasDrawCommandList::FCanvasDrawCommand::CreateLambda(
+			CommandList->EnqueueCommand(USlateIMCanvasDrawCommandList::FCanvasDrawCommand::CreateLambda(
 				[RenderTexture, Triangles](UCanvas* InCanvas)
 				{
 					InCanvas->K2_DrawTriangle(RenderTexture, *Triangles);
@@ -459,9 +459,9 @@ namespace SlateIM::Canvas
 
 	void DrawMaterialTriangles(UMaterialInterface* RenderMaterial, TSharedRef<TArray<FCanvasUVTri>> Triangles)
 	{
-		if (USlateIMEngineCanvasDrawCommandList* CommandList = GetCanvasCommandList())
+		if (USlateIMCanvasDrawCommandList* CommandList = GetCanvasCommandList())
 		{
-			CommandList->EnqueueCommand(USlateIMEngineCanvasDrawCommandList::FCanvasDrawCommand::CreateLambda(
+			CommandList->EnqueueCommand(USlateIMCanvasDrawCommandList::FCanvasDrawCommand::CreateLambda(
 				[RenderMaterial, Triangles](UCanvas* InCanvas)
 				{
 					InCanvas->K2_DrawMaterialTriangle(RenderMaterial, *Triangles);
@@ -472,9 +472,9 @@ namespace SlateIM::Canvas
 
 	void DrawPolygon(UTexture* RenderTexture, const FVector2D& CanvasPosition, const FVector2D& Radius, int32 NumberOfSides, FLinearColor RenderColor)
 	{
-		if (USlateIMEngineCanvasDrawCommandList* CommandList = GetCanvasCommandList())
+		if (USlateIMCanvasDrawCommandList* CommandList = GetCanvasCommandList())
 		{
-			CommandList->EnqueueCommand(USlateIMEngineCanvasDrawCommandList::FCanvasDrawCommand::CreateLambda(
+			CommandList->EnqueueCommand(USlateIMCanvasDrawCommandList::FCanvasDrawCommand::CreateLambda(
 				[RenderTexture, CanvasPosition, Radius, NumberOfSides, RenderColor](UCanvas* InCanvas)
 				{
 					InCanvas->K2_DrawPolygon(
diff --git a/Engine/Plugins/Experimental/SlateIM/Source/SlateIM/Private/Implementation/SlateIM_Containers.cpp b/Engine/Plugins/Experimental/SlateIM/Source/SlateIM/Private/Implementation/SlateIM_Containers.cpp
--- a/Engine/Plugins/Experimental/SlateIM/Source/SlateIM/Private/Implementation/SlateIM_Containers.cpp
+++ b/Engine/Plugins/Experimental/SlateIM/Source/SlateIM/Private/Implementation/SlateIM_Containers.cpp
@@ -94,13 +94,8 @@ namespace SlateIM
 		FSlateIMManager::Get().PopContainer<SImWrapBox>();
 	}
 
-	void BeginHorizontalStack(bool bMaximizeContent)
+	void BeginHorizontalStack()
 	{
-		if (bMaximizeContent)
-		{
-			SlateIM::Maximize();
-		}
-
 		BeginStack(Orient_Horizontal);
 	}
 
@@ -109,13 +104,8 @@ namespace SlateIM
 		EndStack();
 	}
 
-	void BeginVerticalStack(bool bMaximizeContent)
+	void BeginVerticalStack()
 	{
-		if (bMaximizeContent)
-		{
-			SlateIM::Maximize();
-		}
-
 		BeginStack(Orient_Vertical);
 	}
 
@@ -124,13 +114,8 @@ namespace SlateIM
 		EndStack();
 	}
 
-	void BeginHorizontalWrap(bool bMaximizeContent)
+	void BeginHorizontalWrap()
 	{
-		if (bMaximizeContent)
-		{
-			SlateIM::Maximize();
-		}
-
 		BeginWrap(Orient_Horizontal);
 	}
 
@@ -139,13 +124,8 @@ namespace SlateIM
 		EndWrap();
 	}
 
-	void BeginVerticalWrap(bool bMaximizeContent)
+	void BeginVerticalWrap()
 	{
-		if (bMaximizeContent)
-		{
-			SlateIM::Maximize();
-		}
-
 		BeginWrap(Orient_Vertical);
 	}
 
@@ -155,17 +135,12 @@ namespace SlateIM
 	}
 
 	void BeginBorder(const FSlateBrush* BackgroundImage, EOrientation Orientation, bool bAbsorbMouse, FMargin ContentPadding)
-	{
-		BeginBorder(BackgroundImage, {.Orientation = Orientation, .bAbsorbMouse = bAbsorbMouse, .ContentPadding = ContentPadding});
-	}
-
-	void BeginBorder(const FSlateBrush* BackgroundImage, const FBorderParams& Params)
 	{
 		TSharedPtr<SImCompoundWidget> ContainerWidget;
 
 		// Note widget scope must fall out of scope before we push container or else we are pushing this new container into itself instead of into its parent
 		{
-			FWidgetScope<SImCompoundWidget> Scope(Params.ContentPadding, HAlign_Fill, VAlign_Fill, false);
+			FWidgetScope<SImCompoundWidget> Scope(ContentPadding, HAlign_Fill, VAlign_Fill, false);
 			ContainerWidget = Scope.GetWidget();
 
 			if (!ContainerWidget)
@@ -177,9 +152,9 @@ namespace SlateIM
 
 
 			ContainerWidget->SetBackgroundImage(BackgroundImage);
-			ContainerWidget->SetContentPadding(Params.ContentPadding);
-			ContainerWidget->SetAbsorbMouse(Params.bAbsorbMouse);
-			ContainerWidget->SetOrientation(Params.Orientation);
+			ContainerWidget->SetContentPadding(ContentPadding);
+			ContainerWidget->SetAbsorbMouse(bAbsorbMouse);
+			ContainerWidget->SetOrientation(Orientation);
 		}
 
 		FSlateIMManager::Get().PushContainer(FContainerNode(ContainerWidget));
@@ -187,12 +162,7 @@ namespace SlateIM
 
 	void BeginBorder(const FName BorderStyleName, EOrientation Orientation, bool bAbsorbMouse, FMargin ContentPadding)
 	{
-		BeginBorder(FAppStyle::GetBrush(BorderStyleName), {.Orientation = Orientation, .bAbsorbMouse = bAbsorbMouse, .ContentPadding = ContentPadding});
-	}
-
-	void BeginBorder(const FName BorderStyleName, const FBorderParams& Params)
-	{
-		BeginBorder(FAppStyle::GetBrush(BorderStyleName), Params);
+		BeginBorder(FAppStyle::GetBrush(BorderStyleName), Orientation, bAbsorbMouse, ContentPadding);
 	}
 
 	void EndBorder()
@@ -201,11 +171,6 @@ namespace SlateIM
 	}
 
 	bool BeginScrollBox(EOrientation Orientation)
-	{
-		return BeginScrollBox({.Orientation = Orientation});
-	}
-
-	bool BeginScrollBox(const FScrollBoxParams& Params)
 	{
 		TSharedPtr<SImScrollBox> ContainerWidget;
 
@@ -218,7 +183,7 @@ namespace SlateIM
 			{
 				ContainerWidget =
 					SNew(SImScrollBox)
-					.Orientation(Params.Orientation)
+					.Orientation(Orientation)
 					.OnUserScrolled_Lambda([ActivationData = Scope.GetOrCreateActivationMetadata().ToWeakPtr()](float Offset)
 					{
 						FSlateIMManager::Get().ActivateWidget(ActivationData.Pin());
@@ -242,22 +207,12 @@ namespace SlateIM
 		FSlateIMManager::Get().PopContainer<SImScrollBox>();
 	}
 
-	void BeginPopUp(const FName BorderStyleName, const FPopUpParams& Params)
-	{
-		BeginPopUp(FAppStyle::GetBrush(BorderStyleName), Params);
-	}
-
 	void BeginPopUp(const FName BorderStyleName, EOrientation Orientation, bool bAbsorbMouse, FMargin ContentPadding)
 	{
-		BeginPopUp(FAppStyle::GetBrush(BorderStyleName), {.Orientation = Orientation, .bAbsorbMouse = bAbsorbMouse, .ContentPadding = ContentPadding});
+		BeginPopUp(FAppStyle::GetBrush(BorderStyleName), Orientation, bAbsorbMouse, ContentPadding);
 	}
 
 	void BeginPopUp(const FSlateBrush* BorderBrush, EOrientation Orientation, bool bAbsorbMouse, FMargin ContentPadding)
-	{
-		BeginPopUp(BorderBrush, {.Orientation = Orientation, .bAbsorbMouse = bAbsorbMouse, .ContentPadding = ContentPadding});
-	}
-
-	void BeginPopUp(const FSlateBrush* BorderBrush, const FPopUpParams& Params)
 	{
 		TSharedPtr<SImPopUp> ContainerWidget;
 		{
@@ -281,7 +236,7 @@ namespace SlateIM
 		}
 
 		FSlateIMManager::Get().PushContainer(FContainerNode(ContainerWidget));
-		SlateIM::BeginBorder(BorderBrush, Params.Orientation, Params.bAbsorbMouse, Params.ContentPadding);
+		SlateIM::BeginBorder(BorderBrush, Orientation, bAbsorbMouse, ContentPadding);
 	}
 
 	void EndPopUp()
@@ -538,12 +493,7 @@ namespace SlateIM
 		BeginTableColumn(ColumnID);
 	}
 
-	void AddTableColumn(const FName& ColumnID, const FStringView& Label)
-	{
-		AddTableColumn(ColumnID, {.Label = Label});
-	}
-
-	void AddTableColumn(const FName& ColumnID, const FTableColumnParams& Params)
+	void AddTableColumn(const FName& ColumnID, const FStringView& ColumnLabel)
 	{
 		FSlateIMAlignmentState AlignmentState = FSlateIMManager::Get().SaveAlignmentState();
 
@@ -555,10 +505,10 @@ namespace SlateIM
 
 		NextTableColumn(ColumnID);
 
-		if (!Params.Label.IsEmpty())
+		if (!ColumnLabel.IsEmpty())
 		{
 			FSlateIMManager::Get().RestoreAlignmentState(AlignmentState);
-			SlateIM::Text(Params.Label);
+			SlateIM::Text(ColumnLabel);
 		}
 
 		EndTableColumn();
@@ -694,10 +644,6 @@ namespace SlateIM
 	}
 
 	bool BeginTableRowChildren(uint32 ParentRowId, bool bDefaultExpanded)
-	{
-		return BeginTableRowChildren({.ParentRowId = ParentRowId, .bDefaultExpanded = bDefaultExpanded});
-	}
-	bool BeginTableRowChildren(const FTableRowChildrenParams& Params)
 	{
 		// If there's an open cell, close it
 		if (TSharedPtr<SImStackBox> Cell = FSlateIMManager::Get().GetCurrentContainer<SImStackBox>())
@@ -711,37 +657,37 @@ namespace SlateIM
 			return false;
 		}
 
-		ParentRow->SetRowId(Params.ParentRowId);
+		ParentRow->SetRowId(ParentRowId);
 
 		SlateIM::BeginTableRow();
 
-		if (Params.ParentRowId != 0)
+		if (ParentRowId != 0)
 		{
 			if (TSharedPtr<SImTableView> OwningTable = ParentRow->GetOwningTable())
 			{
 				// If the default expanded state is true, we have a state because the default is unexpanded.
-				bool bHasExpansionState = Params.bDefaultExpanded;
-				bool bExpanded = Params.bDefaultExpanded;
+				bool bHasExpansionState = bDefaultExpanded;
+				bool bExpanded = bDefaultExpanded;
 				TSharedPtr<ISlateIMContainer> GrandParent = ParentRow->GetParent();
 
 				if (GrandParent->IsA<FSlateIMTableBody>())
 				{
 					TSharedPtr<FSlateIMTableBody> GrandParentRow = StaticCastSharedPtr<FSlateIMTableBody>(GrandParent);
 					
-					if (GrandParentRow->HasSavedExpansionState(Params.ParentRowId))
+					if (GrandParentRow->HasSavedExpansionState(ParentRowId))
 					{
 						bHasExpansionState = true;
-						bExpanded = GrandParentRow->GetSavedExpansionState(Params.ParentRowId, Params.bDefaultExpanded);
+						bExpanded = GrandParentRow->GetSavedExpansionState(ParentRowId, bDefaultExpanded);
 					}
 				}
 				else
 				{
 					TSharedPtr<FSlateIMTableRow> GrandParentBody = StaticCastSharedPtr<FSlateIMTableRow>(GrandParent);
 
-					if (GrandParentBody->HasSavedExpansionState(Params.ParentRowId))
+					if (GrandParentBody->HasSavedExpansionState(ParentRowId))
 					{
 						bHasExpansionState = true;
-						bExpanded = GrandParentBody->GetSavedExpansionState(Params.ParentRowId, Params.bDefaultExpanded);
+						bExpanded = GrandParentBody->GetSavedExpansionState(ParentRowId, bDefaultExpanded);
 					}
 				}
 
diff --git a/Engine/Plugins/Experimental/SlateIM/Source/SlateIM/Private/Implementation/SlateIM_Graphing.cpp b/Engine/Plugins/Experimental/SlateIM/Source/SlateIM/Private/Implementation/SlateIM_Graphing.cpp
--- a/Engine/Plugins/Experimental/SlateIM/Source/SlateIM/Private/Implementation/SlateIM_Graphing.cpp
+++ b/Engine/Plugins/Experimental/SlateIM/Source/SlateIM/Private/Implementation/SlateIM_Graphing.cpp
@@ -1,4 +1,4 @@
-// Copyright Epic Games, Inc. All Rights Reserved.
+﻿// Copyright Epic Games, Inc. All Rights Reserved.
 
 #include "SlateIM.h"
 
@@ -45,30 +45,20 @@ namespace SlateIM
 	}
 
 	void GraphLine(const TArrayView<FVector2D>& Points, const FLinearColor& LineColor, float LineThickness, const FDoubleRange& XViewRange, const FDoubleRange& YViewRange)
-	{
-		GraphLine(Points, {.XViewRange = XViewRange, .YViewRange = YViewRange, .LineColor = LineColor, .LineThickness = LineThickness});
-	}
-
-	void GraphLine(const TArrayView<FVector2D>& Points, const FGraphLinePointsParams& Params)
 	{
 		TSharedPtr<SImGraph> GraphWidget = GetCurrentGraphWidget();
 		if (ensureMsgf(GraphWidget, TEXT("Calling SlateIM::GraphLine() but the last widget is not a graph widget")))
 		{
-			GraphWidget->AddLineGraph(Points, Params);
+			GraphWidget->AddLineGraph(Points, LineColor, LineThickness, XViewRange, YViewRange);
 		}
 	}
 
 	void GraphLine(const TArrayView<double>& Values, const FLinearColor& LineColor, float LineThickness, const FDoubleRange& ViewRange)
-	{
-		GraphLine(Values, {.ViewRange = ViewRange, .LineColor = LineColor, .LineThickness = LineThickness});
-	}
-
-	void GraphLine(const TArrayView<double>& Values, const FGraphLineValuesParams& Params)
 	{
 		TSharedPtr<SImGraph> GraphWidget = GetCurrentGraphWidget();
 		if (ensureMsgf(GraphWidget, TEXT("Calling SlateIM::GraphLine() but the last widget is not a graph widget")))
 		{
-			GraphWidget->AddLineGraph(Values, Params);
+			GraphWidget->AddLineGraph(Values, LineColor, LineThickness, ViewRange);
 		}
 	}
 }
diff --git a/Engine/Plugins/Experimental/SlateIM/Source/SlateIM/Private/Implementation/SlateIM_Menu.cpp b/Engine/Plugins/Experimental/SlateIM/Source/SlateIM/Private/Implementation/SlateIM_Menu.cpp
--- a/Engine/Plugins/Experimental/SlateIM/Source/SlateIM/Private/Implementation/SlateIM_Menu.cpp
+++ b/Engine/Plugins/Experimental/SlateIM/Source/SlateIM/Private/Implementation/SlateIM_Menu.cpp
@@ -105,7 +105,7 @@ namespace SlateIM
 
 		SlateIM::Padding(0);
 
-		if (Button(InMenuName, {.Style = &FSlateIMStyle::Get().GetWidgetStyle<FButtonStyle>("SlateIM.MenuButton")}))
+		if (Button(InMenuName, &FSlateIMStyle::Get().GetWidgetStyle<FButtonStyle>("SlateIM.MenuButton")))
 		{
 			if (TSharedPtr<SImContextMenuAnchor> Menu = FSlateIMManager::Get().GetCurrentMenuRoot())
 			{
@@ -143,42 +143,27 @@ namespace SlateIM
 	}
 
 	bool AddMenuButton(const FStringView& RowText, const FStringView& ToolTipText)
-	{
-		return AddMenuButton(RowText, {.ToolTipText = ToolTipText});
-	}
-
-	bool AddMenuButton(const FStringView& RowText, const FMenuButtonParams& Params)
 	{
 		TSharedPtr<SImContextMenuAnchor> Anchor = FSlateIMManager::Get().GetCurrentMenuRoot();
 		checkf(Anchor, TEXT("Cannot add menu items without a current active menu anchor"));
 
-		return Anchor->AddMenuButton(RowText, Params.ToolTipText);
+		return Anchor->AddMenuButton(RowText, ToolTipText);
 	}
 
 	bool AddMenuToggleButton(const FStringView& RowText, bool& InOutCurrentState, const FStringView& ToolTipText)
-	{
-		return AddMenuToggleButton(RowText, InOutCurrentState, {.ToolTipText = ToolTipText});
-	}
-
-	bool AddMenuToggleButton(const FStringView& RowText, bool& InOutCurrentState, const FMenuButtonParams& Params)
 	{
 		TSharedPtr<SImContextMenuAnchor> Anchor = FSlateIMManager::Get().GetCurrentMenuRoot();
 		checkf(Anchor, TEXT("Cannot add menu items without a current active menu anchor"));
 
-		return Anchor->AddMenuToggleButton(RowText, InOutCurrentState, Params.ToolTipText);
+		return Anchor->AddMenuToggleButton(RowText, InOutCurrentState, ToolTipText);
 	}
 
 	bool AddMenuCheckButton(const FStringView& RowText, bool& InOutCurrentState, const FStringView& ToolTipText)
-	{
-		return AddMenuCheckButton(RowText, InOutCurrentState, {.ToolTipText = ToolTipText});
-	}
-
-	bool AddMenuCheckButton(const FStringView& RowText, bool& InOutCurrentState, const FMenuButtonParams& Params)
 	{
 		TSharedPtr<SImContextMenuAnchor> Anchor = FSlateIMManager::Get().GetCurrentMenuRoot();
 		checkf(Anchor, TEXT("Cannot add menu items without a current active menu anchor"));
 
-		return Anchor->AddMenuCheckButton(RowText, InOutCurrentState, Params.ToolTipText);
+		return Anchor->AddMenuCheckButton(RowText, InOutCurrentState, ToolTipText);
 	}
 
 	void BeginSubMenu(const FStringView& SubMenuText)
diff --git a/Engine/Plugins/Experimental/SlateIM/Source/SlateIM/Private/Implementation/SlateIM_Roots.cpp b/Engine/Plugins/Experimental/SlateIM/Source/SlateIM/Private/Implementation/SlateIM_Roots.cpp
--- a/Engine/Plugins/Experimental/SlateIM/Source/SlateIM/Private/Implementation/SlateIM_Roots.cpp
+++ b/Engine/Plugins/Experimental/SlateIM/Source/SlateIM/Private/Implementation/SlateIM_Roots.cpp
@@ -84,11 +84,6 @@ namespace SlateIM
 
 #if WITH_ENGINE
 	bool BeginViewportRoot(FName UniqueName, UGameViewportClient* ViewportClient, const FViewportRootLayout& Layout)
-	{
-		return BeginViewportRoot(UniqueName, ViewportClient, {.Layout = Layout});
-	}
-
-	bool BeginViewportRoot(FName UniqueName, UGameViewportClient* ViewportClient, const FViewportParams& Params)
 	{
 		if (!FSlateIMManager::Get().CanUpdateSlateIM())
 		{
@@ -118,7 +113,7 @@ namespace SlateIM
 			ViewportRoot = StaticCastSharedPtr<FSlateIMViewportRoot>(RootNode->RootWidget);
 		}
 
-		ViewportRoot->UpdateViewport(Params.Layout);
+		ViewportRoot->UpdateViewport(Layout);
 
 		FSlateIMManager::Get().BeginRoot(UniqueName);
 
@@ -126,11 +121,6 @@ namespace SlateIM
 	}
 
 	bool BeginViewportRoot(FName UniqueName, ULocalPlayer* LocalPlayer, const FViewportRootLayout& Layout)
-	{
-		return BeginViewportRoot(UniqueName, LocalPlayer, {.Layout = Layout});
-	}
-
-	bool BeginViewportRoot(FName UniqueName, ULocalPlayer* LocalPlayer, const FViewportParams& Params)
 	{
 		if (!FSlateIMManager::Get().CanUpdateSlateIM())
 		{
@@ -160,7 +150,7 @@ namespace SlateIM
 			ViewportRoot = StaticCastSharedPtr<FSlateIMViewportRoot>(RootNode->RootWidget);
 		}
 
-		ViewportRoot->UpdateViewport(Params.Layout);
+		ViewportRoot->UpdateViewport(Layout);
 
 		FSlateIMManager::Get().BeginRoot(UniqueName);
 
@@ -169,11 +159,6 @@ namespace SlateIM
 
 #if WITH_EDITOR
 	bool BeginViewportRoot(FName UniqueName, TSharedPtr<IAssetViewport> AssetViewport, const FViewportRootLayout& Layout)
-	{
-		return BeginViewportRoot(UniqueName, AssetViewport, {.Layout = Layout});
-	}
-
-	bool BeginViewportRoot(FName UniqueName, TSharedPtr<IAssetViewport> AssetViewport, const FViewportParams& Params)
 	{
 		if (!FSlateIMManager::Get().CanUpdateSlateIM())
 		{
@@ -203,7 +188,7 @@ namespace SlateIM
 			ViewportRoot = StaticCastSharedPtr<FSlateIMViewportRoot>(RootNode->RootWidget);
 		}
 
-		ViewportRoot->UpdateViewport(Params.Layout);
+		ViewportRoot->UpdateViewport(Layout);
 
 		FSlateIMManager::Get().BeginRoot(UniqueName);
 
diff --git a/Engine/Plugins/Experimental/SlateIM/Source/SlateIM/Private/Implementation/SlateIM_Tabs.cpp b/Engine/Plugins/Experimental/SlateIM/Source/SlateIM/Private/Implementation/SlateIM_Tabs.cpp
--- a/Engine/Plugins/Experimental/SlateIM/Source/SlateIM/Private/Implementation/SlateIM_Tabs.cpp
+++ b/Engine/Plugins/Experimental/SlateIM/Source/SlateIM/Private/Implementation/SlateIM_Tabs.cpp
@@ -87,11 +87,6 @@ namespace SlateIM
 	}
 
 	bool BeginTab(const FName& TabId, const FSlateIcon& TabIcon, const FText& TabTitle)
-	{
-		return BeginTab(TabId, {.TabIcon = TabIcon, .TabTitle = TabTitle});
-	}
-
-	bool BeginTab(const FName& TabId, const FTabParams& Params)
 	{
 		TSharedPtr<SImTabGroup> TabGroup = FSlateIMManager::Get().GetCurrentContainer<SImTabGroup>();
 		if (!ensureMsgf(TabGroup, TEXT("Cannot begin a Tab - the current container is not a Tab Group. Are your Begin and End function calls mismatched?")))
@@ -115,9 +110,9 @@ namespace SlateIM
 				ContainerWidget->SetOnTabDrawerClosed(FSimpleDelegate::CreateSP(ContainerWidget.Get(), &SImTab::OnDrawerClosed));
 				Scope.UpdateWidget(ContainerWidget);
 
-				// Params.TabIcon and Params.TabTitle are not updated after construction, this could be supported in the future but is not currently part of the Tab Manager API
-				ContainerWidget->TabIcon = Params.TabIcon;
-				ContainerWidget->TabTitle = Params.TabTitle;
+				// TabIcon and TabTitle are not updated after construction, this could be supported in the future but is not currently part of the Tab Manager API
+				ContainerWidget->TabIcon = TabIcon;
+				ContainerWidget->TabTitle = TabTitle;
 			}
 
 			if (!FSlateIMManager::Get().GetMutableCurrentRoot().Tabs.Contains(TabId))
diff --git a/Engine/Plugins/Experimental/SlateIM/Source/SlateIM/Private/Implementation/SlateIM_Util.cpp b/Engine/Plugins/Experimental/SlateIM/Source/SlateIM/Private/Implementation/SlateIM_Util.cpp
--- a/Engine/Plugins/Experimental/SlateIM/Source/SlateIM/Private/Implementation/SlateIM_Util.cpp
+++ b/Engine/Plugins/Experimental/SlateIM/Source/SlateIM/Private/Implementation/SlateIM_Util.cpp
@@ -1,4 +1,4 @@
-// Copyright Epic Games, Inc. All Rights Reserved.
+﻿// Copyright Epic Games, Inc. All Rights Reserved.
 
 #include "Misc/MessageDialog.h"
 #include "SlateIM.h"
@@ -28,11 +28,6 @@ namespace SlateIM
 	}
 
 	EAppReturnType::Type ModalDialog(EAppMsgType::Type MessageType, const FStringView& DialogText, EAppMsgCategory Category, const FStringView& DialogTitle)
-	{
-		return ModalDialog(MessageType, DialogText, {.Category = Category, .Title = DialogTitle});
-	}
-
-	EAppReturnType::Type ModalDialog(EAppMsgType::Type MessageType, const FStringView& DialogText, const FModalDialogParams& Params)
 	{
 		FSlateIMManager::Get().OnSlateIMModalOpened();
 
@@ -40,6 +35,6 @@ namespace SlateIM
 		{
 			FSlateIMManager::Get().OnSlateIMModalClosed();
 		};
-		return FMessageDialog::Open(Params.Category, MessageType, FText::FromStringView(DialogText), FText::FromStringView(Params.Title));
+		return FMessageDialog::Open(Category, MessageType, FText::FromStringView(DialogText), FText::FromStringView(DialogTitle));
 	}
 }
diff --git a/Engine/Plugins/Experimental/SlateIM/Source/SlateIM/Private/Implementation/SlateIM_Widgets.cpp b/Engine/Plugins/Experimental/SlateIM/Source/SlateIM/Private/Implementation/SlateIM_Widgets.cpp
--- a/Engine/Plugins/Experimental/SlateIM/Source/SlateIM/Private/Implementation/SlateIM_Widgets.cpp
+++ b/Engine/Plugins/Experimental/SlateIM/Source/SlateIM/Private/Implementation/SlateIM_Widgets.cpp
@@ -5,7 +5,6 @@
 
 #include "Misc/SlateIMLogging.h"
 #include "Misc/SlateIMManager.h"
-#include "Misc/SlateIMSlateResources.h"
 #include "Misc/SlateIMSlotData.h"
 #include "Misc/SlateIMWidgetScope.h"
 #include "Widgets/Input/SSearchableComboBox.h"
@@ -35,46 +34,35 @@ namespace SlateIM
 	}
 
 	void Text(const FStringView& InText, FSlateColor Color, const FTextBlockStyle* TextStyle /*= nullptr*/)
-	{
-		Text(InText, {.Color = Color, .Style = TextStyle});
-	}
-
-	void Text(const FStringView& InText, const FTextParams& Params)
 	{
 		FWidgetScope<STextBlock> Scope;
-		TSharedPtr<STextBlock> TextBlock = Scope.GetOrCreate([&InText, Params] ()
+		TSharedPtr<STextBlock> TextBlock = Scope.GetOrCreate([&InText, Color, TextStyle] ()
 			{
 				TSharedPtr<STextBlock> TextBlock =
 					SNew(STextBlock)
 					.AutoWrapText(true)
 					.Text(FText::FromStringView(InText))
-					.ColorAndOpacity(Params.Color);
+					.ColorAndOpacity(Color);
 
-				if (Params.Style)
+				if (TextStyle)
 				{
-					TextBlock->SetTextStyle(Params.Style);
+					TextBlock->SetTextStyle(TextStyle);
 				}
 				return TextBlock;
 			});
 
-		Scope.HashData(Params.Color);
-		Scope.HashData(Params.Style);
+		Scope.HashData(Color);
+		Scope.HashData(TextStyle);
 		Scope.HashStringView(InText);
-
-		if (Scope.IsDataHashDirty())
+		if(Scope.IsDataHashDirty())
 		{
 			TextBlock->SetText(FText::FromStringView(InText));
-			TextBlock->SetTextStyle(Params.Style);
-			TextBlock->SetColorAndOpacity(Params.Color);
+			TextBlock->SetTextStyle(TextStyle);
+			TextBlock->SetColorAndOpacity(Color);
 		}
 	}
 
 	bool EditableText(FString& InOutText, const FStringView& HintText, const FEditableTextBoxStyle* TextStyle)
-	{
-		return EditableText(InOutText, {.HintText = HintText, .Style = TextStyle});
-	}
-
-	bool EditableText(FString& InOutText, const FEditableTextParams& Params)
 	{
 		FWidgetScope<SEditableTextBox> Scope;
 		TSharedPtr<SEditableTextBox> EditableText = Scope.GetWidget();
@@ -82,9 +70,9 @@ namespace SlateIM
 
 		bool bWasActivated = false;
 
-		Scope.HashData(Params.Style);
+		Scope.HashData(TextStyle);
 		Scope.HashStringView(InOutText);
-		Scope.HashStringView(Params.HintText);
+		Scope.HashStringView(HintText);
 
 		if (!EditableText)
 		{
@@ -94,7 +82,7 @@ namespace SlateIM
 				SNew(SEditableTextBox)
 				.MinDesiredWidth(MinWidth)
 				.Text(FText::FromStringView(InOutText))
-				.HintText(FText::FromStringView(Params.HintText))
+				.HintText(FText::FromStringView(HintText))
 				.OnTextChanged_Lambda([ActivationData](const FText& NewText)
 				{
 					FSlateIMManager::Get().ActivateWidget(ActivationData.Pin());
@@ -104,9 +92,9 @@ namespace SlateIM
 					FSlateIMManager::Get().ActivateWidget(ActivationData.Pin());
 				});
 
-			if (Params.Style)
+			if (TextStyle)
 			{
-				EditableText->SetStyle(Params.Style);
+				EditableText->SetStyle(TextStyle);
 			}
 			
 			Scope.UpdateWidget(EditableText);
@@ -119,8 +107,8 @@ namespace SlateIM
 
 			if (bIsHashDirty)
 			{
-				EditableText->SetStyle(Params.Style);
-				EditableText->SetHintText(FText::FromStringView(Params.HintText));
+				EditableText->SetStyle(TextStyle);
+				EditableText->SetHintText(FText::FromStringView(HintText));
 			}
 			
 			if (bWasActivated)
@@ -138,7 +126,7 @@ namespace SlateIM
 		return bWasActivated;
 	}
 
-	void Image_Internal(const FImageParams& Params, TFunctionRef<const FSlateBrush*(const TSharedRef<SImage>&)> GetBrush)
+	void Image_Internal(const FSlateColor& ColorAndOpacity, FVector2D DesiredSize, TFunctionRef<const FSlateBrush*(const TSharedRef<SImage>&)> GetBrush)
 	{
 		FWidgetScope<SImage> Scope(Defaults::Padding, HAlign_Left, VAlign_Top);
 		TSharedPtr<SImage> ImageWidget = Scope.GetWidget();
@@ -147,8 +135,8 @@ namespace SlateIM
 		{
 			ImageWidget =
 				SNew(SImage)
-				.ColorAndOpacity(Params.ColorAndOpacity)
-				.DesiredSizeOverride(Params.DesiredSize == FVector2D::ZeroVector ? TOptional<FVector2D>() : Params.DesiredSize);
+				.ColorAndOpacity(ColorAndOpacity)
+				.DesiredSizeOverride(DesiredSize == FVector2D::ZeroVector ? TOptional<FVector2D>() : DesiredSize);
 
 			ImageWidget->SetImage(GetBrush(ImageWidget.ToSharedRef()));
 			Scope.UpdateWidget(ImageWidget);
@@ -156,49 +144,29 @@ namespace SlateIM
 		else
 		{
 			ImageWidget->SetImage(GetBrush(ImageWidget.ToSharedRef()));
-			ImageWidget->SetColorAndOpacity(Params.ColorAndOpacity);
-			ImageWidget->SetDesiredSizeOverride(Params.DesiredSize == FVector2D::ZeroVector ? TOptional<FVector2D>() : Params.DesiredSize);
+			ImageWidget->SetColorAndOpacity(ColorAndOpacity);
+			ImageWidget->SetDesiredSizeOverride(DesiredSize == FVector2D::ZeroVector ? TOptional<FVector2D>() : DesiredSize);
 		}
 	}
 
 	void Image(const FSlateBrush* ImageBrush, const FSlateColor& ColorAndOpacity, FVector2D DesiredSize)
 	{
-		Image_Internal({.ColorAndOpacity = ColorAndOpacity, .DesiredSize = DesiredSize}, [ImageBrush](const TSharedRef<SImage>&) { return ImageBrush; });
-	}
-
-	void Image(const FSlateBrush* ImageBrush, const FImageParams& Params)
-	{
-		Image_Internal(Params, [ImageBrush](const TSharedRef<SImage>&) { return ImageBrush; });
+		Image_Internal(ColorAndOpacity, DesiredSize, [ImageBrush](const TSharedRef<SImage>&) { return ImageBrush; });
 	}
 
 	void Image(const FName ImageStyleName, const FSlateColor& ColorAndOpacity, FVector2D DesiredSize)
 	{
-		Image(FAppStyle::Get().GetBrush(ImageStyleName), {.ColorAndOpacity = ColorAndOpacity, .DesiredSize = DesiredSize});
-	}
-
-	void Image(const FName ImageStyleName, const FImageParams& Params)
-	{
-		Image(FAppStyle::Get().GetBrush(ImageStyleName), Params);
+		Image(FAppStyle::Get().GetBrush(ImageStyleName), ColorAndOpacity, DesiredSize);
 	}
 
 	void Image(const FSlateColor& ColorAndOpacity, FVector2D DesiredSize)
-	{
-		Image({.ColorAndOpacity = ColorAndOpacity, .DesiredSize = DesiredSize});
-	}
-
-	void Image(const FImageParams& Params)
 	{
 		static FSlateBrush DefaultBrush;
-		Image(&DefaultBrush, Params);
+		Image(&DefaultBrush, ColorAndOpacity, DesiredSize);
 	}
 
 #if WITH_ENGINE
 	void Image(UTexture2D* ImageTexture, const FSlateColor& ColorAndOpacity, FVector2D DesiredSize)
-	{
-		Image(ImageTexture, {.ColorAndOpacity = ColorAndOpacity, .DesiredSize = DesiredSize});
-	}
-
-	void Image(UTexture2D* ImageTexture, const FImageParams& Params)
 	{
 		auto MakeTextureBrush = [ImageTexture]() -> FSlateBrush
 		{
@@ -212,7 +180,7 @@ namespace SlateIM
 			return FSlateNoResource();
 		};
 		
-		Image_Internal(Params, [&MakeTextureBrush, ImageTexture](const TSharedRef<SImage>& ImageWidget)
+		Image_Internal(ColorAndOpacity, DesiredSize, [&MakeTextureBrush, ImageTexture](const TSharedRef<SImage>& ImageWidget)
 		{
 			TSharedPtr<FUObjectImageResource> Resource = ImageWidget->GetMetaData<FUObjectImageResource>();
 			if (!Resource)
@@ -230,11 +198,6 @@ namespace SlateIM
 	}
 
 	void Image(UTextureRenderTarget2D* ImageRenderTarget, const FSlateColor& ColorAndOpacity, FVector2D DesiredSize)
-	{
-		Image(ImageRenderTarget, {.ColorAndOpacity = ColorAndOpacity, .DesiredSize = DesiredSize});
-	}
-
-	void Image(UTextureRenderTarget2D* ImageRenderTarget, const FImageParams& Params)
 	{
 		auto MakeRTBrush = [ImageRenderTarget]() -> FSlateBrush
 		{
@@ -248,7 +211,7 @@ namespace SlateIM
 			return FSlateNoResource();
 		};
 		
-		Image_Internal(Params, [&MakeRTBrush, ImageRenderTarget](const TSharedRef<SImage>& ImageWidget)
+		Image_Internal(ColorAndOpacity, DesiredSize, [&MakeRTBrush, ImageRenderTarget](const TSharedRef<SImage>& ImageWidget)
 		{
 			TSharedPtr<FUObjectImageResource> Resource = ImageWidget->GetMetaData<FUObjectImageResource>();
 			if (!Resource)
@@ -267,12 +230,7 @@ namespace SlateIM
 
 	void Image(UMaterialInterface* ImageMaterial, FVector2D BrushSize, const FSlateColor& ColorAndOpacity, FVector2D DesiredSize)
 	{
-		Image(ImageMaterial, BrushSize, {.ColorAndOpacity = ColorAndOpacity, .DesiredSize = DesiredSize});
-	}
-
-	void Image(UMaterialInterface* ImageMaterial, FVector2D BrushSize, const FImageParams& Params)
-	{
-		auto MakeMaterialBrush = [ImageMaterial, BrushSize]() -> FSlateBrush
+		auto MakeMaterialBrush = [ImageMaterial, &BrushSize]() -> FSlateBrush
 		{
 			if (ImageMaterial)
 			{
@@ -281,7 +239,7 @@ namespace SlateIM
 			return FSlateNoResource();
 		};
 		
-		Image_Internal(Params, [&MakeMaterialBrush, ImageMaterial](const TSharedRef<SImage>& ImageWidget)
+		Image_Internal(ColorAndOpacity, DesiredSize, [&MakeMaterialBrush, ImageMaterial](const TSharedRef<SImage>& ImageWidget)
 		{
 			TSharedPtr<FUObjectImageResource> Resource = ImageWidget->GetMetaData<FUObjectImageResource>();
 			if (!Resource)
@@ -301,26 +259,11 @@ namespace SlateIM
 
 	bool Button(const FStringView& InText, const FButtonStyle* InStyle)
 	{
-		return Button(InText, {.Style = InStyle});
-	}
-
-	bool Button(const FStringView& InText, const bool bEnabled, const FButtonStyle* InStyle)
-	{
-		return Button(InText, {.bEnabled = bEnabled, .Style = InStyle});
-	}
-
-	bool Button(const FStringView& InText, const FButtonParams& Params)
-	{
-		if (!Params.bEnabled)
-		{
-			BeginDisabledState();
-		}
-
 		FWidgetScope<SImButton> Scope;
 		TSharedPtr<SImButton> ButtonWidget = Scope.GetWidget();
 
 		Scope.HashStringView(InText);
-		Scope.HashData(Params.Style);
+		Scope.HashData(InStyle);
 
 		bool bWasClicked = false;
 		if (!ButtonWidget)
@@ -335,9 +278,9 @@ namespace SlateIM
 
 			ButtonWidget->SetText(InText);
 
-			if (Params.Style)
+			if (InStyle)
 			{
-				ButtonWidget->SetButtonStyle(Params.Style);
+				ButtonWidget->SetButtonStyle(InStyle);
 			}
 
 			Scope.UpdateWidget(ButtonWidget);
@@ -349,32 +292,38 @@ namespace SlateIM
 			if (Scope.IsDataHashDirty())
 			{
 				ButtonWidget->SetText(InText);
-				if (Params.Style)
+				if (InStyle)
 				{
-					ButtonWidget->SetButtonStyle(Params.Style);
+					ButtonWidget->SetButtonStyle(InStyle);
 				}
 			}
 		}
 
-		if (!Params.bEnabled)
+		return bWasClicked;
+	}
+
+	bool Button(const FStringView& InText, const bool bEnabled, const FButtonStyle* InStyle)
+	{
+		if (!bEnabled)
 		{
-			EndDisabledState();
+			BeginDisabledState();
 		}
 
-		return bWasClicked;
-	}
+		const bool bResult = Button(InText, InStyle);
 
+		if (!bEnabled)
+		{
+			EndDisabledState();
+		}
 
-	bool CheckBox(const FStringView& InLabel, bool& InOutCurrentState, const FCheckBoxStyle* CheckBoxStyle)
-	{
-		return CheckBox(InOutCurrentState, {.Label = InLabel, .CheckBoxStyle = CheckBoxStyle});
+		return bResult;
 	}
 
-	bool CheckBox(bool& InOutCurrentState, const FCheckBoxParams& Params)
+	bool CheckBox(const FStringView& InText, bool& InOutCurrentState, const FCheckBoxStyle* CheckBoxStyle)
 	{
 		ECheckBoxState CurrentEnumState = InOutCurrentState ? ECheckBoxState::Checked : ECheckBoxState::Unchecked;
 
-		const bool bValueChanged = CheckBox(CurrentEnumState, Params);
+		const bool bValueChanged = CheckBox(InText, CurrentEnumState, CheckBoxStyle);
 		if (bValueChanged)
 		{
 			InOutCurrentState = CurrentEnumState == ECheckBoxState::Checked;
@@ -383,20 +332,15 @@ namespace SlateIM
 		return bValueChanged;
 	}
 
-	bool CheckBox(const FStringView& InLabel, ECheckBoxState& InOutCurrentState, const FCheckBoxStyle* CheckBoxStyle)
-	{
-		return CheckBox(InOutCurrentState, {.Label = InLabel, .CheckBoxStyle = CheckBoxStyle});
-	}
-
-	bool CheckBox(ECheckBoxState& InOutCurrentState, const FCheckBoxParams& Params)
+	bool CheckBox(const FStringView& InText, ECheckBoxState& InOutCurrentState, const FCheckBoxStyle* CheckBoxStyle)
 	{
 		FWidgetScope<SImCheckBox> Scope;
 		TSharedPtr<SImCheckBox> CheckboxWidget = Scope.GetWidget();
 
 		bool bValueChanged = false;
 
-		Scope.HashStringView(Params.Label);
-		Scope.HashData(Params.CheckBoxStyle);
+		Scope.HashStringView(InText);
+		Scope.HashData(CheckBoxStyle);
 
 		if (!CheckboxWidget)
 		{
@@ -408,11 +352,11 @@ namespace SlateIM
 					FSlateIMManager::Get().ActivateWidget(ActivationData.Pin());
 				});
 
-			CheckboxWidget->SetText(FText::FromStringView(Params.Label));
+			CheckboxWidget->SetText(FText::FromStringView(InText));
 
-			if (Params.CheckBoxStyle)
+			if (CheckBoxStyle)
 			{
-				CheckboxWidget->SetStyle(Params.CheckBoxStyle);
+				CheckboxWidget->SetStyle(CheckBoxStyle);
 			}
 
 			Scope.UpdateWidget(CheckboxWidget);
@@ -431,10 +375,10 @@ namespace SlateIM
 
 			if (Scope.IsDataHashDirty())
 			{
-				CheckboxWidget->SetText(FText::FromStringView(Params.Label));
-				if (Params.CheckBoxStyle)
+				CheckboxWidget->SetText(FText::FromStringView(InText));
+				if (CheckBoxStyle)
 				{
-					CheckboxWidget->SetStyle(Params.CheckBoxStyle);
+					CheckboxWidget->SetStyle(CheckBoxStyle);
 				}
 			}
 		}
@@ -443,12 +387,12 @@ namespace SlateIM
 	}
 
 	template<typename NumericType>
-	bool SpinBox_Internal(NumericType& InOutValue, const FSpinBoxParams<NumericType>& Params)
+	bool SpinBox_Internal(NumericType& InOutValue, TOptional<NumericType> Min, TOptional<NumericType> Max, const FSpinBoxStyle* SpinBoxStyle = nullptr)
 	{
 		FWidgetScope<SSpinBox<NumericType>> Scope(Defaults::Padding, Defaults::HAlign, Defaults::VAlign, Defaults::bAutoSize, Defaults::InputWidgetWidth);
 		TSharedPtr<SSpinBox<NumericType>> SpinBoxWidget = Scope.GetWidget();
 
-		Scope.HashData(Params.Style);
+		Scope.HashData(SpinBoxStyle);
 
 		class FSpinBoxState : public ISlateMetaData
 		{
@@ -464,8 +408,8 @@ namespace SlateIM
 			TWeakPtr<FSlateIMWidgetActivationMetadata> ActivationData = Scope.GetOrCreateActivationMetadata().ToWeakPtr();
 			SpinBoxWidget =
 				SNew(SSpinBox<NumericType>)
-				.MinValue(Params.Min)
-				.MaxValue(Params.Max)
+				.MinValue(Min)
+				.MaxValue(Max)
 				.Value(InOutValue)
 				.OnBeginSliderMovement_Lambda([SpinBoxState](){ SpinBoxState->bIsChanging = true; })
 				.OnEndSliderMovement_Lambda([SpinBoxState, ActivationData](NumericType NewVal)
@@ -485,15 +429,15 @@ namespace SlateIM
 			SpinBoxWidget->AddMetadata(SpinBoxState);
 			Scope.UpdateWidget(SpinBoxWidget);
 
-			if (Params.Style)
+			if (SpinBoxStyle)
 			{
-				SpinBoxWidget->SetWidgetStyle(Params.Style);
+				SpinBoxWidget->SetWidgetStyle(SpinBoxStyle);
 			}
 		}
 		else
 		{
-			SpinBoxWidget->SetMinValue(Params.Min);
-			SpinBoxWidget->SetMaxValue(Params.Max);
+			SpinBoxWidget->SetMinValue(Min);
+			SpinBoxWidget->SetMaxValue(Max);
 
 			bValueChanged = Scope.IsActivatedThisFrame();
 			if (bValueChanged)
@@ -509,9 +453,9 @@ namespace SlateIM
 				}
 			}
 
-			if (Params.Style && Scope.IsDataHashDirty())
+			if (SpinBoxStyle && Scope.IsDataHashDirty())
 			{
-				SpinBoxWidget->SetWidgetStyle(Params.Style);
+				SpinBoxWidget->SetWidgetStyle(SpinBoxStyle);
 				SpinBoxWidget->InvalidateStyle();
 			}
 		}
@@ -522,74 +466,51 @@ namespace SlateIM
 
 	bool SpinBox(float& InOutValue, TOptional<float> Min, TOptional<float> Max, const FSpinBoxStyle* SpinBoxStyle)
 	{
-		return SpinBox_Internal<float>(InOutValue, {.Min = Min, .Max = Max, .Style = SpinBoxStyle});
+		return SpinBox_Internal<float>(InOutValue, Min, Max, SpinBoxStyle);
 	}
 
 	bool SpinBox(double& InOutValue, TOptional<double> Min, TOptional<double> Max, const FSpinBoxStyle* SpinBoxStyle)
 	{
-		return SpinBox_Internal<double>(InOutValue, {.Min = Min, .Max = Max, .Style = SpinBoxStyle});
+		return SpinBox_Internal<double>(InOutValue, Min, Max, SpinBoxStyle);
 	}
 
 	bool SpinBox(int32& InOutValue, TOptional<int32> Min, TOptional<int32> Max, const FSpinBoxStyle* SpinBoxStyle)
 	{
-		return SpinBox_Internal<int32>(InOutValue, {.Min = Min, .Max = Max, .Style = SpinBoxStyle});
-	}
-
-	template<>
-	bool SpinBox<float>(float& InOutValue, const FSpinBoxParams<float>& Params)
-	{
-		return SpinBox_Internal<float>(InOutValue, Params);
-	}
-
-	template<>
-	bool SpinBox<double>(double& InOutValue, const FSpinBoxParams<double>& Params)
-	{
-		return SpinBox_Internal<double>(InOutValue, Params);
-	}
-
-	template<>
-	bool SpinBox<int32>(int32& InOutValue, const FSpinBoxParams<int32>& Params)
-	{
-		return SpinBox_Internal<int32>(InOutValue, Params);
+		return SpinBox_Internal<int32>(InOutValue, Min, Max, SpinBoxStyle);
 	}
 
 	bool Slider(float& InOutValue, float Min, float Max, float Step, const FSliderStyle* SliderStyle /*= nullptr*/)
-	{
-		return Slider(InOutValue, {.Min = Min, .Max = Max, .Step = Step, .Style = SliderStyle});
-	}
-
-	bool Slider(float& InOutValue, const FSliderParams& Params)
 	{
 		FWidgetScope<SSlider> Scope(Defaults::Padding, Defaults::HAlign, Defaults::VAlign, Defaults::bAutoSize, Defaults::InputWidgetWidth);
 		TSharedPtr<SSlider> SliderWidget = Scope.GetWidget();
 
-		Scope.HashData(Params.Style);
+		Scope.HashData(SliderStyle);
 
 		bool bValueChanged = false;
 		if (!SliderWidget)
 		{
 			SliderWidget =
 				SNew(SSlider)
-				.MinValue(Params.Min)
-				.MaxValue(Params.Max)
-				.StepSize(Params.Step)
+				.MinValue(Min)
+				.MaxValue(Max)
+				.StepSize(Step)
 				.Value(InOutValue)
 				.OnValueChanged_Lambda([ActivationData = Scope.GetOrCreateActivationMetadata().ToWeakPtr()](float NewVal)
 				{
 					FSlateIMManager::Get().ActivateWidget(ActivationData.Pin());
 				});
 
-			if (Params.Style)
+			if (SliderStyle)
 			{
-				SliderWidget->SetStyle(Params.Style);
+				SliderWidget->SetStyle(SliderStyle);
 			}
 
 			Scope.UpdateWidget(SliderWidget);
 		}
 		else
 		{
-			SliderWidget->SetMinAndMaxValues(Params.Min, Params.Max);
-			SliderWidget->SetStepSize(Params.Step);
+			SliderWidget->SetMinAndMaxValues(Min, Max);
+			SliderWidget->SetStepSize(Step);
 
 			bValueChanged = Scope.IsActivatedThisFrame();
 			if (bValueChanged)
@@ -601,9 +522,9 @@ namespace SlateIM
 				SliderWidget->SetValue(InOutValue);
 			}
 			
-			if (Scope.IsDataHashDirty() && Params.Style)
+			if (Scope.IsDataHashDirty() && SliderStyle)
 			{
-				SliderWidget->SetStyle(Params.Style);
+				SliderWidget->SetStyle(SliderStyle);
 			}
 		}
 
@@ -611,16 +532,11 @@ namespace SlateIM
 	}
 
 	void ProgressBar(TOptional<float> Percent, const FProgressBarStyle* ProgressBarStyle /*= nullptr*/)
-	{
-		ProgressBar(Percent, {.Style = ProgressBarStyle});
-	}
-
-	void ProgressBar(TOptional<float> Percent, const FProgressBarParams& Params)
 	{
 		FWidgetScope<SProgressBar> Scope(Defaults::Padding, Defaults::HAlign, Defaults::VAlign, Defaults::bAutoSize, Defaults::InputWidgetWidth);
 		TSharedPtr<SProgressBar> ProgressBarWidget = Scope.GetWidget();
 
-		Scope.HashData(Params.Style);
+		Scope.HashData(ProgressBarStyle);
 
 		if (!ProgressBarWidget)
 		{
@@ -628,9 +544,9 @@ namespace SlateIM
 				SNew(SProgressBar)
 				.Percent(Percent);
 
-			if (Params.Style)
+			if (ProgressBarStyle)
 			{
-				ProgressBarWidget->SetStyle(Params.Style);
+				ProgressBarWidget->SetStyle(ProgressBarStyle);
 			}
 
 			Scope.UpdateWidget(ProgressBarWidget);
@@ -639,14 +555,14 @@ namespace SlateIM
 		{
 			ProgressBarWidget->SetPercent(Percent);
 			
-			if (Scope.IsDataHashDirty() && Params.Style)
+			if (Scope.IsDataHashDirty() && ProgressBarStyle)
 			{
-				ProgressBarWidget->SetStyle(Params.Style);
+				ProgressBarWidget->SetStyle(ProgressBarStyle);
 			}
 		}
 	}
 
-	bool SearchableComboBox_Impl(const TArray<FString>& ComboItems, int32& InOutSelectedItemIndex, const FComboBoxParams& Params)
+	bool SearchableComboBox(const TArray<FString>& ComboItems, int32& InOutSelectedItemIndex, bool bForceRefresh, const FComboBoxStyle* InComboStyle)
 	{
 		FWidgetScope<SSearchableComboBox> Scope;
 		TSharedPtr<SSearchableComboBox> ComboWidget = Scope.GetWidget();
@@ -675,7 +591,7 @@ namespace SlateIM
 			SAssignNew(ComboWidget, SSearchableComboBox)
 			.OptionsSource(&WidgetIMData->Data)
 			.InitiallySelectedItem(WidgetIMData->Data.IsValidIndex(InOutSelectedItemIndex) ? WidgetIMData->Data[InOutSelectedItemIndex] : nullptr)
-			.ComboBoxStyle(Params.Style ? Params.Style : &FAppStyle::Get().GetWidgetStyle<FComboBoxStyle>("ComboBox"))
+			.ComboBoxStyle(InComboStyle ? InComboStyle : &FAppStyle::Get().GetWidgetStyle<FComboBoxStyle>("ComboBox"))
 			.Content()[ComboText]
 			.OnSelectionChanged_Lambda([ActivationData = Scope.GetOrCreateActivationMetadata().ToWeakPtr()](TSharedPtr<FString> NewVal, ESelectInfo::Type SelectType)
 			{
@@ -712,7 +628,7 @@ namespace SlateIM
 		else
 		{
 			TSharedPtr<FComboBoxData> ComboBoxData = ComboWidget->GetMetaData<FComboBoxData>();
-			if (Params.bForceRefresh)
+			if (bForceRefresh)
 			{
 				ComboBoxData->Data = BuildComboOptions();
 				if (InOutSelectedItemIndex != INDEX_NONE && !ComboWidget->IsOpen() && ComboBoxData->Data.IsValidIndex(InOutSelectedItemIndex))
@@ -736,17 +652,12 @@ namespace SlateIM
 		return bValueChanged;
 	}
 
-	bool SearchableComboBox(const TArray<FString>& ComboItems, int32& InOutSelectedItemIndex, bool bForceRefresh, const FComboBoxStyle* InComboStyle)
-	{
-		return SearchableComboBox_Impl(ComboItems, InOutSelectedItemIndex, {.bForceRefresh = bForceRefresh, .Style = InComboStyle});
-	}
-
-	bool ComboBox_Impl(const TArray<FString>& ComboItems, int32& InOutSelectedItemIndex, const FComboBoxParams& Params)
+	bool ComboBox(const TArray<FString>& ComboItems, int32& InOutSelectedItemIndex, bool bForceRefresh, const FComboBoxStyle* InComboStyle)
 	{
 		FWidgetScope<STextComboBox> Scope;
 		TSharedPtr<STextComboBox> ComboWidget = Scope.GetWidget();
 
-		Scope.HashData(Params.Style);
+		Scope.HashData(InComboStyle);
 		Scope.HashData(InOutSelectedItemIndex);
 		
 		using FComboBoxData = FSlateIMDataStore<TArray<TSharedPtr<FString>>>;
@@ -777,9 +688,9 @@ namespace SlateIM
 					FSlateIMManager::Get().ActivateWidget(ActivationData.Pin());
 				});
 
-			if (Params.Style)
+			if (InComboStyle)
 			{
-				ComboWidget->SetStyle(Params.Style);
+				ComboWidget->SetStyle(InComboStyle);
 			}
 			
 			ComboWidget->AddMetadata(WidgetIMData);
@@ -789,7 +700,7 @@ namespace SlateIM
 		else
 		{
 			TSharedPtr<FComboBoxData> ComboBoxData = ComboWidget->GetMetaData<FComboBoxData>();
-			if (Params.bForceRefresh)
+			if (bForceRefresh)
 			{
 				ComboBoxData->Data = BuildComboOptions();
 				if (InOutSelectedItemIndex != INDEX_NONE && !ComboWidget->IsOpen() && ComboBoxData->Data.IsValidIndex(InOutSelectedItemIndex))
@@ -818,34 +729,14 @@ namespace SlateIM
 					ComboWidget->SetSelectedItem(ComboBoxData->Data[InOutSelectedItemIndex]);
 				}
 
-				ComboWidget->SetStyle(Params.Style);
+				ComboWidget->SetStyle(InComboStyle);
 			}
 		}
 
 		return bValueChanged;
 	}
 
-	bool ComboBox(const TArray<FString>& ComboItems, int32& InOutSelectedItemIndex, bool bForceRefresh, const FComboBoxStyle* InComboStyle)
-	{
-		return ComboBox_Impl(ComboItems, InOutSelectedItemIndex, {.bForceRefresh = bForceRefresh, .Style = InComboStyle});
-	}
-
-	bool ComboBox(const TArray<FString>& ComboItems, int32& InOutSelectedItemIndex, const FComboBoxParams& Params)
-	{
-		if (Params.bSearchable)
-		{
-			return SearchableComboBox_Impl(ComboItems, InOutSelectedItemIndex, Params);
-		}
-
-		return ComboBox_Impl(ComboItems, InOutSelectedItemIndex, Params);
-	}
-
 	bool SelectionList(const TArray<FString>& ListItems, int32& InOutSelectedItemIndex, bool bForceRefresh, const FTableViewStyle* InStyle)
-	{
-		return SelectionList(ListItems, InOutSelectedItemIndex, {.bForceRefresh = bForceRefresh, .Style = InStyle});
-	}
-
-	bool SelectionList(const TArray<FString>& ListItems, int32& InOutSelectedItemIndex, const FSelectionListParams& Params)
 	{
 		SCOPED_NAMED_EVENT_TEXT("SlateIM::SelectionList", FColorList::Goldenrod);
 		using ListViewType = SListView<TSharedPtr<FString>>;
@@ -854,7 +745,7 @@ namespace SlateIM
 		FWidgetScope<ListViewType> Scope;
 		TSharedPtr<ListViewType> ListWidget = Scope.GetWidget();
 
-		Scope.HashData(Params.Style);
+		Scope.HashData(InStyle);
 
 		auto BuildListItems = [&]()
 		{
@@ -873,7 +764,7 @@ namespace SlateIM
 			TSharedRef<FListViewData> ListViewData = MakeShared<FListViewData>(BuildListItems());
 			ListWidget = SNew(ListViewType)
 				.SelectionMode(ESelectionMode::Single)
-				.ListViewStyle(Params.Style)
+				.ListViewStyle(InStyle)
 				.ListItemsSource(&(ListViewData->Data))
 				.OnSelectionChanged_Lambda([ActivationData = Scope.GetOrCreateActivationMetadata().ToWeakPtr()](TSharedPtr<FString> NewVal, ESelectInfo::Type SelectType)
 				{
@@ -900,7 +791,7 @@ namespace SlateIM
 		else
 		{
 			TSharedPtr<FListViewData> ListViewData = ListWidget->GetMetaData<FListViewData>();
-			if (Params.bForceRefresh)
+			if (bForceRefresh)
 			{
 				ListViewData->Data = BuildListItems();
 				ListWidget->RequestListRefresh();
@@ -928,7 +819,7 @@ namespace SlateIM
 
 			if (Scope.IsDataHashDirty())
 			{
-				ListWidget->SetStyle(Params.Style);
+				ListWidget->SetStyle(InStyle);
 			}
 		}
 
diff --git a/Engine/Plugins/Experimental/SlateIM/Source/SlateIM/Private/Misc/SlateIMExamples.cpp b/Engine/Plugins/Experimental/SlateIM/Source/SlateIM/Private/Misc/SlateIMExamples.cpp
--- a/Engine/Plugins/Experimental/SlateIM/Source/SlateIM/Private/Misc/SlateIMExamples.cpp
+++ b/Engine/Plugins/Experimental/SlateIM/Source/SlateIM/Private/Misc/SlateIMExamples.cpp
@@ -18,7 +18,6 @@
 #if WITH_EDITOR
 #include "SLevelViewport.h"
 #endif
-
 #if WITH_ENGINE
 #include "Engine/Canvas.h"
 #include "Engine/Engine.h"
@@ -110,7 +109,7 @@ void FSlateIMTestWidget::Draw()
 	SlateIM::AddMenuButton(TEXT("Option 1"));
 	SlateIM::EndMenuBar();
 
-	SlateIM::BeginBorder(FAppStyle::GetBrush("ToolPanel.GroupBorder"), {.bAbsorbMouse = bAbsorbMouse});
+	SlateIM::BeginBorder(FAppStyle::GetBrush("ToolPanel.GroupBorder"), Orient_Vertical, bAbsorbMouse);
 	// Basic perf measurement, outside of the scrollbox so that it's "pinned"
 	if (TimeSinceLastUpdate > 0.5f)
 	{
@@ -217,7 +216,7 @@ void FSlateIMTestWidget::DrawBasics()
 				// EditableText Example
 				{
 					SlateIM::BeginHorizontalStack();
-					SlateIM::EditableText(ComboItemToAdd, {.HintText = TEXT("Add Combo Item")});
+					SlateIM::EditableText(ComboItemToAdd, TEXT("Add Combo Item"));
 					if (SlateIM::IsFocused(SlateIM::EFocusDepth::IncludingDescendants))
 					{
 						SlateIM::BeginPopUp();
@@ -248,22 +247,22 @@ void FSlateIMTestWidget::DrawBasics()
 			}
 
 			// CheckBox examples
-			SlateIM::CheckBox(CheckState, {.Label = TEXT("Check Box")});
+			SlateIM::CheckBox(TEXT("Check Box"), CheckState);
 			if (CheckState)
 			{
 				SlateIM::Text(TEXT("Basic Text"));
-				SlateIM::Text(TEXT("Text With Color"), {.Color = FLinearColor::Green});
-				SlateIM::Text(TEXT("Text With style color"), {.Color = FStyleColors::Primary});
+				SlateIM::Text(TEXT("Text With Color"), FLinearColor::Green);
+				SlateIM::Text(TEXT("Text With style color"), FStyleColors::Primary);
 			}
 	
-			if (SlateIM::CheckBox(CheckStateEnum, {.Label = TEXT("Undetermined Check Box")}))
+			if (SlateIM::CheckBox(TEXT("Undetermined Check Box"), CheckStateEnum))
 			{
 				UE_LOG(LogSlateIM, Log, TEXT("Check Box State Changed"));
 			}
 
 			// ComboBox example
 			SlateIM::BeginHorizontalStack();
-			if (SlateIM::ComboBox(ComboBoxItems, SelectedItemIndex, {.bForceRefresh = bRefreshComboItems}))
+			if (SlateIM::ComboBox(ComboBoxItems, SelectedItemIndex, bRefreshComboItems))
 			{
 				if (ComboBoxItems.IsValidIndex(SelectedItemIndex))
 				{
@@ -274,7 +273,7 @@ void FSlateIMTestWidget::DrawBasics()
 			SlateIM::EndHorizontalStack();
 
 			SlateIM::BeginHorizontalStack();
-			if (SlateIM::ComboBox(ComboBoxItems, SearchableComboBoxSelectedItemIndex, {.bForceRefresh = bRefreshComboItems, .bSearchable = true}))
+			if (SlateIM::SearchableComboBox(ComboBoxItems, SearchableComboBoxSelectedItemIndex, bRefreshComboItems))
 			{
 				if (ComboBoxItems.IsValidIndex(SearchableComboBoxSelectedItemIndex))
 				{
@@ -314,7 +313,7 @@ void FSlateIMTestWidget::DrawBasics()
 			SlateIM::EndVerticalStack();
 
 			// SelectionList example
-			SlateIM::SelectionList(ComboBoxItems, SelectedItem, {.bForceRefresh = bRefreshComboItems});
+			SlateIM::SelectionList(ComboBoxItems, SelectedItem, bRefreshComboItems);
 
 			SlateIM::EndHorizontalStack();
 
@@ -343,7 +342,7 @@ void FSlateIMTestWidget::DrawBasics()
 			SlateIM::BeginVerticalStack();
 			{
 				// Slider example
-				if (SlateIM::Slider(SliderVal, {.Max = SliderMax, .Step = 1.f}))
+				if (SlateIM::Slider(SliderVal, 0, SliderMax, 1))
 				{
 					UE_LOG(LogSlateIM, Log, TEXT("Slider Value Changed [%f]"), SliderVal);
 				}
@@ -354,10 +353,10 @@ void FSlateIMTestWidget::DrawBasics()
 				// SpinBox Examples
 				{
 					SlateIM::Padding(FMargin(0, 10, 0, 5));
-					SlateIM::SpinBox(SliderVal, {.Min = 0.0f, .Max = SliderMax});
+					SlateIM::SpinBox(SliderVal, 0.0f, SliderMax);
 
 					SlateIM::Padding(FMargin(0, 10, 0, 5));
-					SlateIM::SpinBox(IntValue, {.Min = 0, .Max = IntMax});
+					SlateIM::SpinBox(IntValue, 0, IntMax);
 				}
 			}
 			SlateIM::EndVerticalStack();
@@ -402,7 +401,7 @@ void FSlateIMTestWidget::DrawBasics()
 		{
 			SCOPED_NAMED_EVENT_TEXT("DisabledState example", FColorList::Goldenrod);
 			// DisabledState example
-			SlateIM::CheckBox(bShouldBeDisabled, {.Label = TEXT("Disable Everything Below Me")});
+			SlateIM::CheckBox(TEXT("Disable Everything Below Me"), bShouldBeDisabled);
 			if (bShouldBeDisabled)
 			{
 				SlateIM::BeginDisabledState();
@@ -421,36 +420,36 @@ void FSlateIMTestWidget::DrawBasics()
 
 			// This part is only shown if the menu is open
 			SlateIM::AddMenuSection(TEXT("Menu Section 1"));
-			if (SlateIM::AddMenuButton(TEXT("Menu Item 1"), {.ToolTipText = TEXT("Menu Item Tool Tip 1")}))
+			if (SlateIM::AddMenuButton(TEXT("Menu Item 1"), TEXT("Menu Item Tool Tip 1")))
 			{
 				UE_LOG(LogSlateIM, Log, TEXT("Menu Item One menu option clicked"));
 			}
 
-			SlateIM::AddMenuButton(TEXT("Menu Item 2"), {.ToolTipText = TEXT("Menu Item Tool Tip 2")});
-			SlateIM::AddMenuButton(TEXT("Menu Item 3"), {.ToolTipText = TEXT("Menu Item Tool Tip 3")});
-			SlateIM::AddMenuButton(TEXT("Menu Item 4"), {.ToolTipText = TEXT("Menu Item Tool Tip 4")});
+			SlateIM::AddMenuButton(TEXT("Menu Item 2"), TEXT("Menu Item Tool Tip 2"));
+			SlateIM::AddMenuButton(TEXT("Menu Item 3"), TEXT("Menu Item Tool Tip 3"));
+			SlateIM::AddMenuButton(TEXT("Menu Item 4"), TEXT("Menu Item Tool Tip 4"));
 
-			SlateIM::AddMenuCheckButton(TEXT("Menu Item With Check"), MenuCheckState, {.ToolTipText = TEXT("Click to toggle check")});
+			SlateIM::AddMenuCheckButton(TEXT("Menu Item With Check"), MenuCheckState, TEXT("Click to toggle check"));
 
-			if (SlateIM::AddMenuToggleButton(TEXT("Menu Item With Toggle"), MenuToggleState, {.ToolTipText = TEXT("Toggle this box")}))
+			if (SlateIM::AddMenuToggleButton(TEXT("Menu Item With Toggle"), MenuToggleState, TEXT("Toggle this box")))
 			{
 				UE_LOG(LogSlateIM, Log, TEXT("Menu Item With Toggle clicked"));
 			}
 
 			SlateIM::AddMenuSeparator();
 			SlateIM::BeginSubMenu(TEXT("Sub Menu"));
-			SlateIM::AddMenuButton(TEXT("SubMenu Item 1"), {.ToolTipText = TEXT("Menu Item Tool Tip 1")});
-			SlateIM::AddMenuButton(TEXT("SubMenu Item 2"), {.ToolTipText = TEXT("Menu Item Tool Tip 2")});
-			SlateIM::AddMenuButton(TEXT("SubMenu Item 3"), {.ToolTipText = TEXT("Menu Item Tool Tip 3")});
-			SlateIM::AddMenuButton(TEXT("SubMenu Item 4"), {.ToolTipText = TEXT("Menu Item Tool Tip 4")});
+			SlateIM::AddMenuButton(TEXT("SubMenu Item 1"), TEXT("Menu Item Tool Tip 1"));
+			SlateIM::AddMenuButton(TEXT("SubMenu Item 2"), TEXT("Menu Item Tool Tip 2"));
+			SlateIM::AddMenuButton(TEXT("SubMenu Item 3"), TEXT("Menu Item Tool Tip 3"));
+			SlateIM::AddMenuButton(TEXT("SubMenu Item 4"), TEXT("Menu Item Tool Tip 4"));
 			SlateIM::EndSubMenu();
-			SlateIM::AddMenuButton(TEXT("Menu Item 5"), {.ToolTipText = TEXT("A Menu Item between 2 SubMenus")});
-			SlateIM::AddMenuButton(TEXT("Menu Item 6"), {.ToolTipText = TEXT("A Menu Item between 2 SubMenus")});
+			SlateIM::AddMenuButton(TEXT("Menu Item 5"), TEXT("A Menu Item between 2 SubMenus"));
+			SlateIM::AddMenuButton(TEXT("Menu Item 6"), TEXT("A Menu Item between 2 SubMenus"));
 			SlateIM::BeginSubMenu(TEXT("Sub Menu 2"));
-			SlateIM::AddMenuButton(TEXT("SubMenu2 Item 1"), {.ToolTipText = TEXT("Menu Item Tool Tip 1")});
-			SlateIM::AddMenuButton(TEXT("SubMenu2 Item 2"), {.ToolTipText = TEXT("Menu Item Tool Tip 2")});
-			SlateIM::AddMenuButton(TEXT("SubMenu2 Item 3"), {.ToolTipText = TEXT("Menu Item Tool Tip 3")});
-			SlateIM::AddMenuButton(TEXT("SubMenu2 Item 4"), {.ToolTipText = TEXT("Menu Item Tool Tip 4")});
+			SlateIM::AddMenuButton(TEXT("SubMenu2 Item 1"), TEXT("Menu Item Tool Tip 1"));
+			SlateIM::AddMenuButton(TEXT("SubMenu2 Item 2"), TEXT("Menu Item Tool Tip 2"));
+			SlateIM::AddMenuButton(TEXT("SubMenu2 Item 3"), TEXT("Menu Item Tool Tip 3"));
+			SlateIM::AddMenuButton(TEXT("SubMenu2 Item 4"), TEXT("Menu Item Tool Tip 4"));
 			SlateIM::EndSubMenu();
 			SlateIM::EndContextMenuAnchor();
 		}
@@ -580,7 +579,7 @@ void FSlateIMTestWidget::DrawTables()
 			}
 
 			NumTableItems = NumItems;
-			SlateIM::CheckBox(bShouldLiveUpdateTable, {.Label = TEXT("Live Update Table?")});
+			SlateIM::CheckBox(TEXT("Live Update Table?"), bShouldLiveUpdateTable);
 			if (bShouldLiveUpdateTable)
 			{
 				NumTableItems = LiveNumItems;
@@ -605,7 +604,7 @@ void FSlateIMTestWidget::DrawTables()
 				{
 					SlateIM::Padding(FMargin(5, 0));
 					SlateIM::VAlign(VAlign_Center);// Centers the button in the row
-					SlateIM::Text(FString::Printf(TEXT("Item %d/%d"), i + 1, NumItems), {.Color = FColor::MakeRedToGreenColorFromScalar(static_cast<float>(i) / NumItems)});
+					SlateIM::Text(FString::Printf(TEXT("Item %d/%d"), i + 1, NumItems), FColor::MakeRedToGreenColorFromScalar(static_cast<float>(i) / NumItems));
 					SlateIM::Padding(FMargin(0));
 
 					if (SlateIM::Button(TEXT("Click")))
@@ -641,7 +640,7 @@ void FSlateIMTestWidget::DrawTables()
 			SlateIM::InitialTableColumnWidth(80.f);
 			SlateIM::Maximize();
 			SlateIM::HAlign(HAlign_Center);
-			SlateIM::AddTableColumn(TEXT("Button"), {.Label = TEXT("Button")});
+			SlateIM::AddTableColumn(TEXT("Button"), TEXT("Button"));
 			SlateIM::EndTableHeader();
 			SlateIM::BeginTableBody();
 			for (int32 i = 0; i < NumTableItems; ++i)
@@ -656,11 +655,11 @@ void FSlateIMTestWidget::DrawTables()
 
 					if (!bIsSelected)
 					{
-						SlateIM::Text(FString::Printf(TEXT("Item %d/%d"), i + 1, NumTableItems), {.Color = FColor::MakeRedToGreenColorFromScalar(static_cast<float>(i) / NumTableItems)});
+						SlateIM::Text(FString::Printf(TEXT("Item %d/%d"), i + 1, NumTableItems), FColor::MakeRedToGreenColorFromScalar(static_cast<float>(i) / NumTableItems));
 					}
 					else
 					{
-						SlateIM::Text(FString::Printf(TEXT("Item %d/%d [Selected]"), i + 1, NumTableItems), {.Color = FColor::MakeRedToGreenColorFromScalar(static_cast<float>(i) / NumTableItems)});
+						SlateIM::Text(FString::Printf(TEXT("Item %d/%d [Selected]"), i + 1, NumTableItems), FColor::MakeRedToGreenColorFromScalar(static_cast<float>(i) / NumTableItems));
 					}
 				}
 
@@ -681,7 +680,7 @@ void FSlateIMTestWidget::DrawTables()
 		
 		// Tree Example
 		{
-			SlateIM::CheckBox(bShouldAddCanada, {.Label = TEXT("Show Canada?")});
+			SlateIM::CheckBox(TEXT("Show Canada?"), bShouldAddCanada);
 
 			SCOPED_NAMED_EVENT_TEXT("Tree Example", FColorList::Goldenrod);
 			const FTableRowStyle* TableRowStyle = &FAppStyle::Get().GetWidgetStyle<FTableRowStyle>("TableView.AlternatingRow");
@@ -690,9 +689,9 @@ void FSlateIMTestWidget::DrawTables()
 			SlateIM::MaxHeight(200.f);
 			SlateIM::HAlign(HAlign_Fill);
 			SlateIM::VAlign(VAlign_Fill);
-			SlateIM::BeginTable({.RowStyle = TableRowStyle});
-			SlateIM::AddTableColumn(TEXT("Name"), {.Label = TEXT("Name")});
-			SlateIM::AddTableColumn(TEXT("Type"), {.Label = TEXT("Type")});
+			SlateIM::BeginTable(nullptr, TableRowStyle);
+			SlateIM::AddTableColumn(TEXT("Name"), TEXT("Name"));
+			SlateIM::AddTableColumn(TEXT("Type"), TEXT("Type"));
 			{
 				if (SlateIM::NextTableCell())
 				{
@@ -732,7 +731,7 @@ void FSlateIMTestWidget::DrawTables()
 							SlateIM::Text(TEXT("Country"));
 						}
 						static constexpr FLazyName CanadaChildren = TEXT("CanadaChildren");
-						if (SlateIM::BeginTableRowChildren({.ParentRowId = GetTypeHash(CanadaChildren)}))
+						if (SlateIM::BeginTableRowChildren(GetTypeHash(CanadaChildren)))
 						{
 							if (SlateIM::NextTableCell())
 							{
@@ -744,7 +743,7 @@ void FSlateIMTestWidget::DrawTables()
 								SlateIM::VAlign(VAlign_Center);
 								SlateIM::Text(TEXT("Province"));
 							}
-							if (SlateIM::BeginTableRowChildren({.ParentRowId = static_cast<uint32>(ChildRowId++)}))
+							if (SlateIM::BeginTableRowChildren(ChildRowId++))
 							{
 								if (SlateIM::NextTableCell())
 								{
@@ -769,7 +768,7 @@ void FSlateIMTestWidget::DrawTables()
 								SlateIM::VAlign(VAlign_Center);
 								SlateIM::Text(TEXT("Province"));
 							}
-							if (SlateIM::BeginTableRowChildren({.ParentRowId = static_cast<uint32>(ChildRowId++)}))
+							if (SlateIM::BeginTableRowChildren(ChildRowId++))
 							{
 								if (SlateIM::NextTableCell())
 								{
@@ -797,13 +796,8 @@ void FSlateIMTestWidget::DrawTables()
 						SlateIM::VAlign(VAlign_Center);
 						SlateIM::Text(TEXT("Country"));
 					}
-#if WITH_ENGINE
 					constexpr uint32 RowID = UE::HashStringFNV1a32("UnitedStatesChildren");
-#else
-					constexpr uint32 RowID = 0;
-#endif
-
-					if (SlateIM::BeginTableRowChildren({.ParentRowId = RowID}))
+					if (SlateIM::BeginTableRowChildren(RowID))
 					{
 						if (SlateIM::NextTableCell())
 						{
@@ -865,12 +859,12 @@ void FSlateIMTestWidget::DrawTables()
 
 		// Dynamically added/removed child row
 		{
-			SlateIM::CheckBox(bShouldAddChildRow, {.Label = TEXT("Show Child?")});
+			SlateIM::CheckBox(TEXT("Show Child?"), bShouldAddChildRow);
 
 			SlateIM::Fill();
 			SlateIM::HAlign(HAlign_Fill);
 			SlateIM::BeginTable();
-			SlateIM::AddTableColumn(TEXT("Column"), {.Label = TEXT("Column")});
+			SlateIM::AddTableColumn(TEXT("Column"), TEXT("Column"));
 
 			if (SlateIM::NextTableCell())
 			{
@@ -913,7 +907,7 @@ void FSlateIMTestWidget::DrawTables()
 			SlateIM::Text(TEXT("Columns:"));
 			for (int32 i = 0; i < Columns.Num(); ++i)
 			{
-				SlateIM::CheckBox(Columns[i], {.Label = GetColumnName(i)});
+				SlateIM::CheckBox(GetColumnName(i), Columns[i]);
 			}
 			SlateIM::EndHorizontalStack();
 
@@ -921,7 +915,7 @@ void FSlateIMTestWidget::DrawTables()
 			SlateIM::Text(TEXT("Rows:"));
 			for (int32 i = 0; i < Rows.Num(); ++i)
 			{
-				SlateIM::CheckBox(Rows[i], {.Label = GetRowName(i)});
+				SlateIM::CheckBox(GetRowName(i), Rows[i]);
 			}
 			SlateIM::EndHorizontalStack();
 
@@ -935,7 +929,7 @@ void FSlateIMTestWidget::DrawTables()
 			{
 				if (Columns[i])
 				{
-					SlateIM::AddTableColumn(FName(GetColumnName(i)), {.Label = GetColumnName(i)});
+					SlateIM::AddTableColumn(FName(GetColumnName(i)), GetColumnName(i));
 				}
 			}
 			for (int32 i = 0; i < Rows.Num(); ++i)
@@ -978,7 +972,7 @@ void FSlateIMTestWidget::DrawGraphs()
 		SlateIM::Fill();
 		SlateIM::MinHeight(200.f);
 		SlateIM::BeginGraph();
-		SlateIM::GraphLine(SquareGraphValues.Compact(), {.ViewRange = FDoubleRange(0, 1.0), .LineColor = FLinearColor::White, .LineThickness = 3.f});
+		SlateIM::GraphLine(SquareGraphValues.Compact(), FLinearColor::White, 3.f, FDoubleRange(0, 1.0));
 		SlateIM::EndGraph();
 
 		SlateIM::VAlign(VAlign_Center);
@@ -1010,9 +1004,9 @@ void FSlateIMTestWidget::DrawGraphs()
 		SlateIM::Fill();
 		SlateIM::MinHeight(200.f);
 		SlateIM::BeginGraph();
-		SlateIM::GraphLine(SinGraphValues.Compact(), {.XViewRange = FDoubleRange(NextSinX - 100.0, NextSinX), .YViewRange = FDoubleRange(-1.5, 1.5), .LineColor = FColor::Orange});
-		SlateIM::GraphLine(CosGraphValues.Compact(), {.XViewRange = FDoubleRange(NextCosX - 100.0, NextCosX), .YViewRange = FDoubleRange(-1.5, 1.5), .LineColor = FLinearColor::Green});
-		SlateIM::GraphLine(TanGraphValues.Compact(), {.XViewRange = FDoubleRange(NextTanX - 100.0, NextTanX), .YViewRange = FDoubleRange(-1.5, 1.5), .LineColor = FColor::Magenta});
+		SlateIM::GraphLine(SinGraphValues.Compact(), FColor::Orange, 1.f, FDoubleRange(NextSinX - 100.0, NextSinX), FDoubleRange(-1.5, 1.5));
+		SlateIM::GraphLine(CosGraphValues.Compact(), FLinearColor::Green, 1.f, FDoubleRange(NextCosX - 100.0, NextCosX), FDoubleRange(-1.5, 1.5));
+		SlateIM::GraphLine(TanGraphValues.Compact(), FColor::Magenta, 1.f, FDoubleRange(NextTanX - 100.0, NextTanX), FDoubleRange(-1.5, 1.5));
 		SlateIM::EndGraph();
 
 		SlateIM::MinWidth(50.f);
@@ -1020,13 +1014,13 @@ void FSlateIMTestWidget::DrawGraphs()
 		SlateIM::BeginVerticalStack();
 		SlateIM::Fill();
 		SlateIM::VAlign(VAlign_Center);
-		SlateIM::Text(FString::Printf(TEXT("%0.3lf"), SinGraphValues.Last().Y), {.Color = FColor::Orange});
+		SlateIM::Text(FString::Printf(TEXT("%0.3lf"), SinGraphValues.Last().Y), FColor::Orange);
 		SlateIM::Fill();
 		SlateIM::VAlign(VAlign_Center);
-		SlateIM::Text(FString::Printf(TEXT("%0.3lf"), CosGraphValues.Last().Y), {.Color = FLinearColor::Green});
+		SlateIM::Text(FString::Printf(TEXT("%0.3lf"), CosGraphValues.Last().Y), FLinearColor::Green);
 		SlateIM::Fill();
 		SlateIM::VAlign(VAlign_Center);
-		SlateIM::Text(FString::Printf(TEXT("%0.3lf"), TanGraphValues.Last().Y), {.Color = FColor::Magenta});
+		SlateIM::Text(FString::Printf(TEXT("%0.3lf"), TanGraphValues.Last().Y), FColor::Magenta);
 		SlateIM::EndVerticalStack();
 	}
 	SlateIM::EndHorizontalStack();
@@ -1159,7 +1153,7 @@ void FSlateIMTestWidget::DrawInputs()
 			SlateIM::BeginBorder(&WBrush);
 			SlateIM::HAlign(HAlign_Center);
 			SlateIM::VAlign(VAlign_Center);
-			SlateIM::Text(TEXT("W"), {.Color = FLinearColor::White});
+			SlateIM::Text(TEXT("W"), FLinearColor::White);
 			SlateIM::EndBorder();
 		}
 		SlateIM::Spacer(FVector2D(50, 50));
@@ -1174,7 +1168,7 @@ void FSlateIMTestWidget::DrawInputs()
 			SlateIM::BeginBorder(&ABrush);
 			SlateIM::HAlign(HAlign_Center);
 			SlateIM::VAlign(VAlign_Center);
-			SlateIM::Text(TEXT("A"), {.Color = FLinearColor::White});
+			SlateIM::Text(TEXT("A"), FLinearColor::White);
 			SlateIM::EndBorder();
 		}
 		{
@@ -1186,7 +1180,7 @@ void FSlateIMTestWidget::DrawInputs()
 			SlateIM::BeginBorder(&SBrush);
 			SlateIM::HAlign(HAlign_Center);
 			SlateIM::VAlign(VAlign_Center);
-			SlateIM::Text(TEXT("S"), {.Color = FLinearColor::White});
+			SlateIM::Text(TEXT("S"), FLinearColor::White);
 			SlateIM::EndBorder();
 		}
 		{
@@ -1198,7 +1192,7 @@ void FSlateIMTestWidget::DrawInputs()
 			SlateIM::BeginBorder(&DBrush);
 			SlateIM::HAlign(HAlign_Center);
 			SlateIM::VAlign(VAlign_Center);
-			SlateIM::Text(TEXT("D"), {.Color = FLinearColor::White});
+			SlateIM::Text(TEXT("D"), FLinearColor::White);
 			SlateIM::EndBorder();
 		}
 		SlateIM::EndHorizontalStack();
@@ -1230,7 +1224,7 @@ void FSlateIMTestWidget::DrawInputs()
 			SlateIM::BeginBorder(&LMBBrush);
 			SlateIM::HAlign(HAlign_Center);
 			SlateIM::VAlign(VAlign_Center);
-			SlateIM::Text(TEXT("LMB"), {.Color = FLinearColor::White});
+			SlateIM::Text(TEXT("LMB"), FLinearColor::White);
 			SlateIM::EndBorder();
 		}
 		{
@@ -1242,7 +1236,7 @@ void FSlateIMTestWidget::DrawInputs()
 			SlateIM::BeginBorder(&RMBBrush);
 			SlateIM::HAlign(HAlign_Center);
 			SlateIM::VAlign(VAlign_Center);
-			SlateIM::Text(TEXT("RMB"), {.Color = FLinearColor::White});
+			SlateIM::Text(TEXT("RMB"), FLinearColor::White);
 			SlateIM::EndBorder();
 		}
 		SlateIM::EndHorizontalStack();
@@ -1283,14 +1277,14 @@ void FSlateIMTestWidget::DrawTabs()
 				SlateIM::TabSplitterSizeCoefficient(0.3f);
 				SlateIM::BeginTabStack();
 				{
-					if (SlateIM::BeginTab(TEXT("Left Tab"), {.TabIcon = FSlateIcon(FAppStyle::GetAppStyleSetName(), TEXT("Symbols.LeftArrow"))}))
+					if (SlateIM::BeginTab(TEXT("Left Tab"), FSlateIcon(FAppStyle::GetAppStyleSetName(), TEXT("Symbols.LeftArrow"))))
 					{
 						// Fill so that the text autowrapping updates properly as the splitter is adjusted by the user
 						SlateIM::HAlign(HAlign_Fill);
 						SlateIM::Text(TEXT("Left Tab taking up 30% of the horizontal space"));
 					}
 					SlateIM::EndTab();
-					if (SlateIM::BeginTab(TEXT("Right Tab"), {.TabIcon = FSlateIcon(FAppStyle::GetAppStyleSetName(), TEXT("Symbols.RightArrow"))}))
+					if (SlateIM::BeginTab(TEXT("Right Tab"), FSlateIcon(FAppStyle::GetAppStyleSetName(), TEXT("Symbols.RightArrow"))))
 					{
 						// Fill so that the text autowrapping updates properly as the splitter is adjusted by the user
 						SlateIM::HAlign(HAlign_Fill);
@@ -1359,7 +1353,7 @@ void FSlateIMTestWidget::DrawTabs()
 			{
 				for (int32 i = 0; i < DynamicTabCount; i++)
 				{
-					if (SlateIM::BeginTab(*FString::Printf(TEXT("DynamicTab%d"), i), {.TabTitle = FText::Format(INVTEXT("Tab {0}"), FText::AsNumber(i + 1))}))
+					if (SlateIM::BeginTab(*FString::Printf(TEXT("DynamicTab%d"), i), FSlateIcon(), FText::Format(INVTEXT("Tab {0}"), FText::AsNumber(i + 1))))
 					{
 						SlateIM::Text(FString::Printf(TEXT("Tab %d Content"), i + 1));
 					}
@@ -1449,7 +1443,7 @@ void FSlateIMTestWidget::DrawFilter()
 
 		// Draw the filter input and filter the list
 		bool bFruitListChanged = false;
-		if (SlateIM::TextFilter(FruitFilter, {.HintText = TEXT("Search fruits...")}))
+		if (SlateIM::TextFilter(FruitFilter, TEXT("Search fruits...")))
 		{
 			FilteredFruitList = FruitFilter.Filter(FruitList);
 			bFruitListChanged = true;
@@ -1459,7 +1453,7 @@ void FSlateIMTestWidget::DrawFilter()
 		SlateIM::Fill();
 		SlateIM::HAlign(HAlign_Fill);
 		int32 SelectedFruitIndex = INDEX_NONE;
-		SlateIM::SelectionList(FilteredFruitList.GetValue(), SelectedFruitIndex, {.bForceRefresh = bFruitListChanged});
+		SlateIM::SelectionList(FilteredFruitList.GetValue(), SelectedFruitIndex, bFruitListChanged);
 	}
 	SlateIM::EndVerticalStack();
 }
@@ -1471,16 +1465,16 @@ void FSlateStyleBrowser::DrawWindow(float DeltaTime)
 	SlateIM::VAlign(VAlign_Fill);
 	SlateIM::Fill();
 	SlateIM::BeginVerticalStack();
-	SlateIM::EditableText(SearchString, {.HintText = TEXT("Search Styles...")});
+	SlateIM::EditableText(SearchString, TEXT("Search Styles..."));
 	
 	SlateIM::HAlign(HAlign_Fill);
 	SlateIM::VAlign(VAlign_Fill);
 	SlateIM::Fill();
 	SlateIM::BeginTable();
-	SlateIM::AddTableColumn(TEXT("Name"), {.Label = TEXT("Name")});
-	SlateIM::AddTableColumn(TEXT("Preview"), {.Label = TEXT("Preview")});
+	SlateIM::AddTableColumn(TEXT("Name"), TEXT("Name"));
+	SlateIM::AddTableColumn(TEXT("Preview"), TEXT("Preview"));
 	SlateIM::InitialTableColumnWidth(80.f);
-	SlateIM::AddTableColumn(TEXT("Button"), {.Label = TEXT("Button")});
+	SlateIM::AddTableColumn(TEXT("Button"), TEXT("Button"));
 
 	const auto& Style = static_cast<const SlateIM::Private::FExposedSlateStyle&>(FAppStyle::Get());
 	const FString LowerCaseSearchString = SearchString.ToLower();
@@ -1561,7 +1555,7 @@ void FSlateStyleBrowser::DrawWindow(float DeltaTime)
 					SlateIM::HAlign(HAlign_Center);
 					SlateIM::VAlign(VAlign_Center);
 					SlateIM::Fill();
-					SlateIM::Text(TEXT("The quick brown fox jumps over the lazy dog."), {.Style = WidgetStyle.Value});
+					SlateIM::Text(TEXT("The quick brown fox jumps over the lazy dog."), WidgetStyle.Value);
 				}
 			
 				DrawCopyCell(WidgetStyle.Key);
@@ -1589,7 +1583,7 @@ void FSlateStyleBrowser::DrawWindow(float DeltaTime)
 					SlateIM::HAlign(HAlign_Center);
 					SlateIM::VAlign(VAlign_Center);
 					SlateIM::Fill();
-					SlateIM::EditableText(PreviewText, {.HintText = TEXT("Hint text..."), .Style = WidgetStyle.Value});
+					SlateIM::EditableText(PreviewText, TEXT("Hint text..."), WidgetStyle.Value);
 				}
 				
 				DrawCopyCell(WidgetStyle.Key);
@@ -1617,7 +1611,7 @@ void FSlateStyleBrowser::DrawWindow(float DeltaTime)
 					SlateIM::HAlign(HAlign_Center);
 					SlateIM::VAlign(VAlign_Center);
 					SlateIM::Fill();
-					SlateIM::Button(TEXT("Click Me"), {.Style = WidgetStyle.Value});
+					SlateIM::Button(TEXT("Click Me"), WidgetStyle.Value);
 				}
 				
 				DrawCopyCell(WidgetStyle.Key);
@@ -1645,7 +1639,7 @@ void FSlateStyleBrowser::DrawWindow(float DeltaTime)
 					SlateIM::HAlign(HAlign_Center);
 					SlateIM::VAlign(VAlign_Center);
 					SlateIM::Fill();
-					SlateIM::SpinBox(SpinBoxValue, {.Min = -100.f, .Max = 100.f, .Style = WidgetStyle.Value});
+					SlateIM::SpinBox(SpinBoxValue, -100.f, 100.f, WidgetStyle.Value);
 				}
 				
 				DrawCopyCell(WidgetStyle.Key);
@@ -1673,7 +1667,7 @@ void FSlateStyleBrowser::DrawWindow(float DeltaTime)
 					SlateIM::HAlign(HAlign_Center);
 					SlateIM::VAlign(VAlign_Center);
 					SlateIM::Fill();
-					SlateIM::Slider(SliderValue, {.Max = 100.f, .Step = 0.1f, .Style = WidgetStyle.Value});
+					SlateIM::Slider(SliderValue, 0.f, 100.f, 0.1f, WidgetStyle.Value);
 				}
 				
 				DrawCopyCell(WidgetStyle.Key);
@@ -1701,7 +1695,7 @@ void FSlateStyleBrowser::DrawWindow(float DeltaTime)
 					SlateIM::HAlign(HAlign_Center);
 					SlateIM::VAlign(VAlign_Center);
 					SlateIM::Fill();
-					SlateIM::ProgressBar(SliderValue / 100.f, {.Style = WidgetStyle.Value});
+					SlateIM::ProgressBar(SliderValue / 100.f, WidgetStyle.Value);
 				}
 				
 				DrawCopyCell(WidgetStyle.Key);
@@ -1729,7 +1723,8 @@ void FSlateStyleBrowser::DrawWindow(float DeltaTime)
 					SlateIM::HAlign(HAlign_Center);
 					SlateIM::VAlign(VAlign_Center);
 					SlateIM::Fill();
-					SlateIM::ComboBox(Options, SelectedComboIndex, {.Style = WidgetStyle.Value});
+					constexpr bool bForceRefresh = false;
+					SlateIM::ComboBox(Options, SelectedComboIndex, bForceRefresh, WidgetStyle.Value);
 				}
 				
 				DrawCopyCell(WidgetStyle.Key);
@@ -1762,14 +1757,14 @@ void FSlateStyleBrowser::DrawWindow(float DeltaTime)
 						SlateIM::HAlign(HAlign_Center);
 						SlateIM::VAlign(VAlign_Top);
 						SlateIM::Fill();
-						SlateIM::SelectionList(Options, SelectedListIndex, {.Style = WidgetStyle.Value});
+						SlateIM::SelectionList(Options, SelectedListIndex, false, WidgetStyle.Value);
 
 						SlateIM::HAlign(HAlign_Center);
 						SlateIM::VAlign(VAlign_Top);
 						SlateIM::Fill();
-						SlateIM::BeginTable({.Style = WidgetStyle.Value});
-						SlateIM::AddTableColumn(TEXT("Column 1"), {.Label = TEXT("Column 1")});
-						SlateIM::AddTableColumn(TEXT("Column 2"), {.Label = TEXT("Column 2")});
+						SlateIM::BeginTable(WidgetStyle.Value);
+						SlateIM::AddTableColumn(TEXT("Column 1"), TEXT("Column 1"));
+						SlateIM::AddTableColumn(TEXT("Column 2"), TEXT("Column 2"));
 						SlateIM::NextTableCell();
 						SlateIM::Text(TEXT("Cell 1"));
 						SlateIM::NextTableCell();
@@ -1816,7 +1811,7 @@ void FSlateIMTestViewportWidget::DrawWidget(float DeltaTime)
 {
 	if (GEngine && GEngine->GameViewport)
 	{
-		if (SlateIM::BeginViewportRoot("SlateIMTestSuiteViewport", GEngine->GameViewport, {.Layout = Layout}))
+		if (SlateIM::BeginViewportRoot("SlateIMTestSuiteViewport", GEngine->GameViewport, Layout))
 		{
 			TestWidget.Draw();
 		}
@@ -1828,7 +1823,7 @@ void FSlateIMTestViewportWidget::DrawWidget(float DeltaTime)
 		TSharedPtr<SLevelViewport> LevelViewport = StaticCastSharedPtr<SLevelViewport>(GCurrentLevelEditingViewportClient->GetEditorViewportWidget());
 		if (LevelViewport.IsValid())
 		{
-			if (SlateIM::BeginViewportRoot("SlateIMTestSuiteViewport", LevelViewport, {.Layout = Layout}))
+			if (SlateIM::BeginViewportRoot("SlateIMTestSuiteViewport", LevelViewport, Layout))
 			{
 				TestWidget.Draw();
 			}
diff --git a/Engine/Plugins/Experimental/SlateIM/Source/SlateIM/Private/Misc/SlateIMSlateResources.h b/Engine/Plugins/Experimental/SlateIM/Source/SlateIM/Private/Misc/SlateIMSlateResources.h
--- a/Engine/Plugins/Experimental/SlateIM/Source/SlateIM/Private/Misc/SlateIMSlateResources.h
+++ b/Engine/Plugins/Experimental/SlateIM/Source/SlateIM/Private/Misc/SlateIMSlateResources.h
@@ -2,6 +2,9 @@
  
 #pragma once
 
+// Entire file is WITH_ENGINE
+#if WITH_ENGINE
+
 #include "Styling/SlateBrush.h"
 #include "Types/ISlateMetaData.h"
 #include "UObject/Object.h"
@@ -84,3 +87,5 @@ namespace SlateIM
 		virtual bool IsOfTypeImpl(const FName& Type) const override { return GetTypeId() == Type || Super::IsOfTypeImpl(Type); }
 	};
 }
+
+#endif // WITH_ENGINE
diff --git a/Engine/Plugins/Experimental/SlateIM/Source/SlateIM/Private/Roots/SlateIMViewportRoot.h b/Engine/Plugins/Experimental/SlateIM/Source/SlateIM/Private/Roots/SlateIMViewportRoot.h
--- a/Engine/Plugins/Experimental/SlateIM/Source/SlateIM/Private/Roots/SlateIMViewportRoot.h
+++ b/Engine/Plugins/Experimental/SlateIM/Source/SlateIM/Private/Roots/SlateIMViewportRoot.h
@@ -4,13 +4,17 @@
 
 #if WITH_ENGINE
 #include "ISlateIMRoot.h"
-#include "SlateIMParametersFwd.h"
 #include "UObject/WeakObjectPtrTemplates.h"
 #include "Widgets/Layout/SConstraintCanvas.h"
 
 class SImWrapper;
 class SScaleBox;
 
+namespace SlateIM
+{
+	struct FViewportRootLayout;
+}
+
 class UGameViewportClient;
 class ULocalPlayer;
 
diff --git a/Engine/Plugins/Experimental/SlateIM/Source/SlateIM/Private/SlateIM_Internal.h b/Engine/Plugins/Experimental/SlateIM/Source/SlateIM/Private/SlateIM_Internal.h
--- a/Engine/Plugins/Experimental/SlateIM/Source/SlateIM/Private/SlateIM_Internal.h
+++ b/Engine/Plugins/Experimental/SlateIM/Source/SlateIM/Private/SlateIM_Internal.h
@@ -3,9 +3,8 @@
 #pragma once
 
 #include "Math/MathFwd.h"
-#include "SlateIMParametersFwd.h"
 #include "Templates/FunctionFwd.h"
-#include "Templates/SharedPointerFwd.h"
+#include "Templates//SharedPointerFwd.h"
 
 class SImage;
 enum EOrientation : int;
@@ -23,5 +22,5 @@ namespace SlateIM
 	void BeginWrap(EOrientation Orientation);
 	void EndWrap();
 
-	void Image_Internal(const FImageParams& Params, TFunctionRef<const FSlateBrush* (const TSharedRef<SImage>&)> GetBrush);
+	void Image_Internal(const FSlateColor& ColorAndOpacity, FVector2D DesiredSize, TFunctionRef<const FSlateBrush* (const TSharedRef<SImage>&)> GetBrush);
 } // SlateIM
diff --git a/Engine/Plugins/Experimental/SlateIM/Source/SlateIM/Private/Widgets/SImGraph.cpp b/Engine/Plugins/Experimental/SlateIM/Source/SlateIM/Private/Widgets/SImGraph.cpp
--- a/Engine/Plugins/Experimental/SlateIM/Source/SlateIM/Private/Widgets/SImGraph.cpp
+++ b/Engine/Plugins/Experimental/SlateIM/Source/SlateIM/Private/Widgets/SImGraph.cpp
@@ -2,7 +2,6 @@
 
 
 #include "SImGraph.h"
-#include "SlateIMParameters.h"
 
 SLATE_IMPLEMENT_WIDGET(SImGraph)
 
@@ -36,29 +35,29 @@ void SImGraph::EndGraph()
 	}
 }
 
-void SImGraph::AddLineGraph(const TArrayView<FVector2D>& Points, const SlateIM::FGraphLinePointsParams& Params)
+void SImGraph::AddLineGraph(const TArrayView<FVector2D>& Points, const FLinearColor& LineColor, float LineThickness, const FDoubleRange& XViewRange, const FDoubleRange& YViewRange)
 {
-	ensure(Params.XViewRange.HasLowerBound() && Params.XViewRange.HasUpperBound());
-	ensure(Params.YViewRange.HasLowerBound() && Params.YViewRange.HasUpperBound());
+	ensure(XViewRange.HasLowerBound() && XViewRange.HasUpperBound());
+	ensure(YViewRange.HasLowerBound() && YViewRange.HasUpperBound());
 	
-	FSlateIMLineGraphData& LineGraph = GetNextLineGraph(Params.LineColor, Params.LineThickness, Params.XViewRange, Params.YViewRange);
+	FSlateIMLineGraphData& LineGraph = GetNextLineGraph(LineColor, LineThickness, XViewRange, YViewRange);
 	LineGraph.NormalizedPoints.Reserve(Points.Num());
 
 	for (const FVector2D& Point : Points)
 	{
 		LineGraph.NormalizedPoints.Emplace(
-			(Point.X - Params.XViewRange.GetLowerBound().GetValue()) / (Params.XViewRange.GetUpperBound().GetValue() - Params.XViewRange.GetLowerBound().GetValue()),
-			(Point.Y - Params.YViewRange.GetLowerBound().GetValue()) / (Params.YViewRange.GetUpperBound().GetValue() - Params.YViewRange.GetLowerBound().GetValue()));
+			(Point.X - XViewRange.GetLowerBound().GetValue()) / (XViewRange.GetUpperBound().GetValue() - XViewRange.GetLowerBound().GetValue()),
+			(Point.Y - YViewRange.GetLowerBound().GetValue()) / (YViewRange.GetUpperBound().GetValue() - YViewRange.GetLowerBound().GetValue()));
 	}
 }
 
-void SImGraph::AddLineGraph(const TArrayView<double>& Values, const SlateIM::FGraphLineValuesParams& Params)
+void SImGraph::AddLineGraph(const TArrayView<double>& Values, const FLinearColor& LineColor, float LineThickness, const FDoubleRange& ViewRange)
 {
-	ensure(Params.ViewRange.HasLowerBound() && Params.ViewRange.HasUpperBound());
+	ensure(ViewRange.HasLowerBound() && ViewRange.HasUpperBound());
 
-	const FDoubleRange XViewRange = FDoubleRange(0.0, static_cast<double>(Values.Num()));
+	const FDoubleRange XViewRange = FDoubleRange(0, static_cast<double>(Values.Num()));
 	
-	FSlateIMLineGraphData& LineGraph = GetNextLineGraph(Params.LineColor, Params.LineThickness, XViewRange, Params.ViewRange);
+	FSlateIMLineGraphData& LineGraph = GetNextLineGraph(LineColor, LineThickness, XViewRange, ViewRange);
 	LineGraph.NormalizedPoints.Reserve(Values.Num());
 
 	double XValue = 0.0;
@@ -66,7 +65,7 @@ void SImGraph::AddLineGraph(const TArrayView<double>& Values, const SlateIM::FGr
 	{
 		LineGraph.NormalizedPoints.Emplace(
 			(XValue - XViewRange.GetLowerBound().GetValue()) / (XViewRange.GetUpperBound().GetValue() - XViewRange.GetLowerBound().GetValue()),
-			(Value - Params.ViewRange.GetLowerBound().GetValue()) / (Params.ViewRange.GetUpperBound().GetValue() - Params.ViewRange.GetLowerBound().GetValue()));
+			(Value - ViewRange.GetLowerBound().GetValue()) / (ViewRange.GetUpperBound().GetValue() - ViewRange.GetLowerBound().GetValue()));
 
 		XValue += 1.0;
 	}
@@ -128,6 +127,5 @@ FSlateIMLineGraphData& SImGraph::GetNextLineGraph(const FLinearColor& LineColor,
 	NewLineGraph.LineThickness = LineThickness;
 	NewLineGraph.XViewRange = XViewRange;
 	NewLineGraph.YViewRange = YViewRange;
-
 	return LineGraphs.Emplace_GetRef(MoveTemp(NewLineGraph));
 }
diff --git a/Engine/Plugins/Experimental/SlateIM/Source/SlateIM/Private/Widgets/SImGraph.h b/Engine/Plugins/Experimental/SlateIM/Source/SlateIM/Private/Widgets/SImGraph.h
--- a/Engine/Plugins/Experimental/SlateIM/Source/SlateIM/Private/Widgets/SImGraph.h
+++ b/Engine/Plugins/Experimental/SlateIM/Source/SlateIM/Private/Widgets/SImGraph.h
@@ -2,7 +2,6 @@
 
 #pragma once
 
-#include "SlateIMParametersFwd.h"
 #include "Widgets/SLeafWidget.h"
 
 #define UE_API SLATEIM_API
@@ -22,7 +21,10 @@ class SImGraph : public SLeafWidget
 	SLATE_DECLARE_WIDGET(SImGraph, SLeafWidget)
 	
 public:
-	SLATE_BEGIN_ARGS(SImGraph) {}
+	SLATE_BEGIN_ARGS(SImGraph)
+		{
+		}
+
 	SLATE_END_ARGS()
 
 	UE_API void Construct(const FArguments& InArgs);
@@ -30,8 +32,8 @@ public:
 	UE_API void BeginGraph();
 	UE_API void EndGraph();
 
-	UE_API void AddLineGraph(const TArrayView<FVector2D>& Points, const SlateIM::FGraphLinePointsParams& Params);
-	UE_API void AddLineGraph(const TArrayView<double>& Values, const SlateIM::FGraphLineValuesParams& Params);
+	UE_API void AddLineGraph(const TArrayView<FVector2D>& Points, const FLinearColor& LineColor, float LineThickness, const FDoubleRange& XViewRange, const FDoubleRange& YViewRange);
+	UE_API void AddLineGraph(const TArrayView<double>& Points, const FLinearColor& LineColor, float LineThickness, const FDoubleRange& ViewRange);
 
 private:
 	UE_API virtual int32 OnPaint( const FPaintArgs& Args, const FGeometry& AllottedGeometry, const FSlateRect& MyCullingRect, FSlateWindowElementList& OutDrawElements, int32 LayerId, const FWidgetStyle& InWidgetStyle, bool bParentEnabled ) const override;
diff --git a/Engine/Plugins/Experimental/SlateIM/Source/SlateIM/Public/SlateIM.h b/Engine/Plugins/Experimental/SlateIM/Source/SlateIM/Public/SlateIM.h
--- a/Engine/Plugins/Experimental/SlateIM/Source/SlateIM/Public/SlateIM.h
+++ b/Engine/Plugins/Experimental/SlateIM/Source/SlateIM/Public/SlateIM.h
@@ -1,11 +1,9 @@
 // Copyright Epic Games, Inc. All Rights Reserved.
 
 #pragma once
-
 #include "Framework/Views/ITypedTableView.h"
 #include "Layout/Margin.h"
 #include "Modules/ModuleInterface.h"
-#include "SlateIMParameters.h"
 #include "Styling/SlateColor.h"
 #include "Styling/StyleColors.h"
 #include "Textures/SlateIcon.h"
@@ -63,10 +61,33 @@ public:
 namespace SlateIM
 {
 #pragma region Roots
+#if WITH_EDITOR || WITH_ENGINE
+	struct FViewportRootLayout
+	{
+		/** How to anchor the root within the viewport (normalized units) */
+		FAnchors Anchors = FAnchors(0);
+		
+		/** Offset the root from the anchored position in the viewport (in slate units) */
+		FVector2f Offset = FVector2f::ZeroVector;
+
+		/** How the root is aligned to its anchor (normalized units) */
+		FVector2f Alignment = FVector2f::ZeroVector;
+
+		/** Optional set the size (in slate units) of the root in the viewport. When unset, the root will auto-size for its content. */
+		TOptional<FVector2f> Size;
+
+		/** The ZOrder of the SlateIM widget in the viewport */
+		int32 ZOrder = 10000;
+
+		/** Global scale to apply within the viewport */
+		float Scale = 1.f;
+	};
+#endif // WITH_EDITOR || WITH_ENGINE
+
 	/**
 	 * Begins a new floating window root
 	 *
-	 * Please use the method below with the struct parameter.
+	 * Please use the method below with the struct parameters.
 	 *
 	 * @param UniqueName a unique name to identify this window
 	 * @param WindowTitle The name of the window. This will also be displayed in the windows title bar
@@ -79,11 +100,26 @@ namespace SlateIM
 	 */
 	SLATEIM_API bool BeginWindowRoot(FName UniqueName, const FStringView& WindowTitle, FVector2f WindowSize, bool bShouldReopen = false);
 
+	struct FWindowParams
+	{
+		// The name of the window. This will also be displayed in the windows title bar
+		FStringView WindowTitle;
+
+		// The size to open the window at (does not update the size of an existing window)
+		FVector2f WindowSize = FVector2f::ZeroVector;
+
+		// If this window was created with a previous call to BeginWindow and then closed, passing in true will reopen the window and false will leave it closed. If the window has never been seen this parameter does nothing and the window will open
+		bool bShouldReopen = false;
+
+		// Window will be create always on top.
+		bool bAlwaysOnTop = false;
+	};
+
 	/**
 	 * Begins a new floating window root
 	 *
 	 * @param UniqueName a unique name to identify this window
-	 * @param Params (optional) Additional parameters to create the window.
+	 * @param Params Configuration parameters for the window.
 	 *
 	 * @return the current window state. True if open, false if closed or not updating
 	 *
@@ -93,96 +129,51 @@ namespace SlateIM
 
 #if WITH_ENGINE
 	// TODO - Game viewport widgets might need something to control HitTestability and cursor visibility
-
-	/**
-	 * Begins a new root in the game viewport
-	 *
-	 * Please use the method below with the struct-parameter.
-	 *
-	 * @param UniqueName a globally unique name to identify this root
-	 * @param ViewportClient the viewport client to add the widget to
-	 * @param Layout How to lay out the root within the viewport
-	 *
-	 * @return Whether the root is valid and updating
-	 */
-	SLATEIM_API bool BeginViewportRoot(FName UniqueName, UGameViewportClient* ViewportClient, const FViewportRootLayout& Layout);
-
-	/**
-	 * Begins a new root in the game viewport
-	 *
-	 * @param UniqueName a globally unique name to identify this root
-	 * @param ViewportClient the viewport client to add the widget to
-	 * @param Params (optional) Additional parameters to create the root.
-	 *
-	 * @return Whether the root is valid and updating
-	 */
-	SLATEIM_API bool BeginViewportRoot(FName UniqueName, UGameViewportClient* ViewportClient, const FViewportParams& Params = FViewportParams());
-	
 	/**
-	 * Begins a new root in the provided player's viewport
-	 *
-	 * Please use the method below with the struct-parameter.
+	 *	Begins a new root in the game viewport
 	 * 
 	 * @param UniqueName a globally unique name to identify this root
-	 * @param LocalPlayer the player whose viewport to add the widget to
+	 * @param ViewportClient the viewport client to add the widget to
 	 * @param Layout How to lay out the root within the viewport
 	 * 
 	 * @return Whether the root is valid and updating
 	 */
-	SLATEIM_API bool BeginViewportRoot(FName UniqueName, ULocalPlayer* LocalPlayer, const FViewportRootLayout& Layout);
+	SLATEIM_API bool BeginViewportRoot(FName UniqueName, UGameViewportClient* ViewportClient, const FViewportRootLayout& Layout = FViewportRootLayout());
 	
 	/**
-	 * Begins a new root in the provided player's viewport
+	 *	Begins a new root in the provided player's viewport
 	 * 
 	 * @param UniqueName a globally unique name to identify this root
 	 * @param LocalPlayer the player whose viewport to add the widget to
-	 * @param Params (optional) Additional parameters to create the root.
+	 * @param Layout How to lay out the root within the viewport
 	 * 
 	 * @return Whether the root is valid and updating
 	 */
-	SLATEIM_API bool BeginViewportRoot(FName UniqueName, ULocalPlayer* LocalPlayer, const FViewportParams& Params = FViewportParams());
+	SLATEIM_API bool BeginViewportRoot(FName UniqueName, ULocalPlayer* LocalPlayer, const FViewportRootLayout& Layout = FViewportRootLayout());
 #endif
 
 #if WITH_EDITOR
 	/**
 	 * Begins a SlateIM root in the provided in editor viewport
-	 *
-	 * Please use the method below with the struct-parameter.
 	 * 
 	 * @param UniqueName Unique identifier for this root. Must be globally unique within the same editor session
 	 * @param AssetViewport The editor viewpor to add the SlateIM content to
 	 * @param Layout How to lay out the root within the viewport
 	 * 
 	 * @return Whether the root is valid and updating
-	 *
-	 * @note Not available in BPFL: TSharedPtr<IAssetViewport> is not Blueprint-compatible
-	 */
-	SLATEIM_API bool BeginViewportRoot(FName UniqueName, TSharedPtr<IAssetViewport> AssetViewport, const FViewportRootLayout& Layout);
-
-	/**
-	 * Begins a SlateIM root in the provided in editor viewport
-	 *
-	 * @param UniqueName Unique identifier for this root. Must be globally unique within the same editor session
-	 * @param AssetViewport The editor viewpor to add the SlateIM content to
-	 * @param Params (optional) Additional parameters to create the root.
-	 *
-	 * @return Whether the root is valid and updating
-	 *
-	 * @note Not available in BPFL: TSharedPtr<IAssetViewport> is not Blueprint-compatible
 	 */
-	SLATEIM_API bool BeginViewportRoot(FName UniqueName, TSharedPtr<IAssetViewport> AssetViewport, const FViewportParams& Params = FViewportParams());
+	SLATEIM_API bool BeginViewportRoot(FName UniqueName, TSharedPtr<IAssetViewport> AssetViewport, const FViewportRootLayout& Layout = FViewportRootLayout());
 #endif
 
 	/**
 	 * Begins a SlateIM root that exposes a slate widget for embedding in other Slate hierarchies
-	 *
+	 * 
 	 * @param UniqueName a globally unique name to identify this root
 	 * @param OutSlateIMWidget where to output the resulting SlateIM widget to
-	 *
+	 * 
 	 * @return Whether the root is valid and updating
-	 *
+	 * 
 	 * @note Consider using FSlateIMExposedBase to handle creating and updating your Slate IM exposed widget
-	 * @note Not available in BPFL: TSharedPtr<SWidget> is not Blueprint-compatible
 	 */
 	SLATEIM_API bool BeginExposedRoot(FName UniqueName, TSharedPtr<SWidget>& OutSlateIMWidget);
 
@@ -195,49 +186,51 @@ namespace SlateIM
 #pragma region Containers
 	/**
 	 * Begins a horizontally stacked container. All widgets created within the Horizontal Stack container will be placed in-order left-to-right
-	 * 
-	 * @param bMaximizeContent (optional) Maximize() the contents of the stack
 	 */
-	SLATEIM_API void BeginHorizontalStack(bool bMaximizeContent = false);
-
+	SLATEIM_API void BeginHorizontalStack();
 	SLATEIM_API void EndHorizontalStack();
 
 	/**
 	 * Begins a vertically stacked container. All widgets created within the Vertical Stack container will be placed in-order top-to-bottom
-	 *
-	 * @param bMaximizeContent (optional) Maximize() the contents of the stack
 	 */
-	SLATEIM_API void BeginVerticalStack(bool bMaximizeContent = false);
+	SLATEIM_API void BeginVerticalStack();
 	SLATEIM_API void EndVerticalStack();
 
 	/**
 	 * Begins a horizontally wrapped container. All widgets created within the Horizontal Wrap container will be placed in-order left-to-right until the allotted width is filled,
 	 * then the content will begin a new line and so-on.
 	 *
-	 * @param bMaximizeContent (optional) Maximize() the contents of the wrap
-	 *
 	 * @note Wrap box children are always AutoSize and do not support Fill
 	 */
-	SLATEIM_API void BeginHorizontalWrap(bool bMaximizeContent = false);
-
+	SLATEIM_API void BeginHorizontalWrap();
 	SLATEIM_API void EndHorizontalWrap();
 	
 	/**
 	 * Begins a vertically wrapped container. All widgets created within the Vertical Wrap container will be placed in-order top-to-bottom until the allotted height is filled,
 	 * then the content will begin a new column and so-on.
 	 *
-	 * @param bMaximizeContent (optional) Maximize() the contents of the wrap
-	 *
 	 * @note Wrap container children are always AutoSize and do not support Fill
 	 */
-	SLATEIM_API void BeginVerticalWrap(bool bMaximizeContent = false);
-
+	SLATEIM_API void BeginVerticalWrap();
 	SLATEIM_API void EndVerticalWrap();
 
+	struct FTableParams
+	{
+		// A Table View Style to override the default style. Changes to the style after the first invocation are not reflected,
+		// provide a new style object to update the visuals at runtime
+		const FTableViewStyle* Style = nullptr;
+		
+		// A Table Row Style to override the default table row style. Changing the style only affects new table rows.
+		const FTableRowStyle* RowStyle = nullptr;
+
+		// Select mode for the table.
+		ESelectionMode::Type SelectionMode = ESelectionMode::None;
+	};
+
 	/**
 	 * Begin a Table container. The contents are automatically scrollable if the table is sized smaller than all its content.
 	 *
-	 * Please use the method below with the struct-parameter.
+	 * Please use the below method instead!
 	 *
 	 * @param InStyle A Table View Style to override the default style. Changes to the style after the first invocation are not reflected,
 	 *				  provide a new style object to update the visuals at runtime
@@ -248,7 +241,7 @@ namespace SlateIM
 	/**
 	 * Begin a Table container. The contents are automatically scrollable if the table is sized smaller than all its content.
 	 *
-	 * @param Params (optional) Additional parameters to create the table.
+	 * @param Params (optional) Settings for table construction.
 	 */
 	SLATEIM_API void BeginTable(const FTableParams& Params = FTableParams());
 
@@ -279,20 +272,10 @@ namespace SlateIM
 	/**
 	 * Adds a column and immediately ends it with the given label as its body.
 	 *
-	 * Please use the method below with the struct-parameter.
-	 *
 	 * @param ColumnID The name to use as ID for the column
 	 * @param ColumnLabel Added a default label to the column
 	 */
-	SLATEIM_API void AddTableColumn(const FName& ColumnID, const FStringView& ColumnLabel);
-
-	/**
-	 * Adds a column and immediately ends it with the given label as its body.
-	 *
-	 * @param ColumnID The name to use as ID for the column
-	 * @param Params (optional) Additional parameters to create the column
-	 */
-	SLATEIM_API void AddTableColumn(const FName& ColumnID, const FTableColumnParams& Params = FTableColumnParams());
+	SLATEIM_API void AddTableColumn(const FName& ColumnID, const FStringView& ColumnLabel = FStringView());
 
 	/**
 	 * Set the column to a fixed width, not resizable by users. Call this before calling AddTableColumn().
@@ -334,8 +317,6 @@ namespace SlateIM
 	 * If the table is a tree, a unique row id must be provided when starting a new set of child rows so that the expansion
 	 * state of the parent row is saved. Not providing a row id will cause the parent row to revert to the default expansion
 	 * state if its grandparent row is closed.
-	 *
-	 * Please use the method below with the struct-parameter.
 	 * 
 	 * @param ParentRowId Id of the parent row. 0 to disable state saving. Row ids must be unique to its parent.
 	 * @param bDefaultExpanded (optional) Set the default expansion state of the parent (whether the children following are visible.)
@@ -343,20 +324,7 @@ namespace SlateIM
 	 * 
 	 * @return Whether the parent row is expanded or not. When false, the child content can be skipped to save cycles.
 	 */
-	SLATEIM_API bool BeginTableRowChildren(uint32 ParentRowId, bool bDefaultExpanded = false);
-
-	/**
-	 * Begin adding content as a child to the current table row. NextTableCell() must still be called after this function.
-	 * If the table is a tree, a unique row id must be provided when starting a new set of child rows so that the expansion
-	 * state of the parent row is saved. Not providing a row id will cause the parent row to revert to the default expansion
-	 * state if its grandparent row is closed.
-	 *
-	 * @param Params (optional) Additional parameters to create the table children.
-	 *
-	 * @return Whether the parent row is expanded or not. When false, the child content can be skipped to save cycles.
-	 */
-	SLATEIM_API bool BeginTableRowChildren(const FTableRowChildrenParams& Params = FTableRowChildrenParams());
-	
+	SLATEIM_API bool BeginTableRowChildren(uint32 ParentRowId = 0, bool bDefaultExpanded = false);
 	/**
 	 * Stop adding child content to the parent row. This must be called no matter the result of BeginTableRowChildren()
 	 */
@@ -364,108 +332,37 @@ namespace SlateIM
 
 	/**
 	 * Begin a container with a background image.
-	 *
-	 * Please use the method below with the struct-parameter.
 	 * 
 	 * @param BorderBrush The brush to use as the background for this container
 	 * @param Orientation Which direction the contents of the container should flow
-	 * @param bAbsorbMouse (optional) Whether the container should handle all mouse inputs
-	 * @param ContentPadding (optional) How much to pad the contents of the container
+	 * @param bAbsorbMouse Whether the container should handle all mouse inputs
+	 * @param ContentPadding How much to pad the contents of the container
 	 */
-	SLATEIM_API void BeginBorder(const FSlateBrush* BorderBrush, EOrientation Orientation, bool bAbsorbMouse = true, FMargin ContentPadding = FMargin(2.0f));
+	SLATEIM_API void BeginBorder(const FSlateBrush* BorderBrush, EOrientation Orientation = Orient_Vertical, bool bAbsorbMouse = true, FMargin ContentPadding = FMargin(2.0f));
 	
 	/**
 	 * Begin a container with a background image.
-	 *
-	 * @param BorderBrush The brush to use as the background for this container
-	 * @param Params (optional) Additional params to create the border
-	 */
-	SLATEIM_API void BeginBorder(const FSlateBrush* BorderBrush, const FBorderParams& Params = FBorderParams());
-
-	/**
-	 * Begin a container with a background image.
-	 *
-	 * Please use the method below with the struct-parameter.
 	 * 
 	 * @param BorderStyleName The name of the brush in FAppStyle to use as the background for this container
 	 * @param Orientation Which direction the contents of the container should flow
-	 * @param bAbsorbMouse (optional) Whether the container should handle all mouse inputs
-	 * @param ContentPadding (optional) How much to pad the contents of the container
-	 */
-	SLATEIM_API void BeginBorder(const FName BorderStyleName, EOrientation Orientation, bool bAbsorbMouse = true, FMargin ContentPadding = FMargin(2.0f));
-
-	/**
-	 * Begin a container with a background image.
-	 *
-	 * @param BorderStyleName The name of the brush in FAppStyle to use as the background for this container
-	 * @param Params (optional) Additional params to create the border
+	 * @param bAbsorbMouse Whether the container should handle all mouse inputs
+	 * @param ContentPadding How much to pad the contents of the container
 	 */
-	SLATEIM_API void BeginBorder(const FName BorderStyleName, const FBorderParams& Params = FBorderParams());
-
+	SLATEIM_API void BeginBorder(const FName BorderStyleName, EOrientation Orientation = Orient_Vertical, bool bAbsorbMouse = true, FMargin ContentPadding = FMargin(2.0f));
 	SLATEIM_API void EndBorder();
 
 	/**
 	 * Begins a container that will allow the user to scroll when its content is larger than the allotted space.
-	 *
-	 * Please use the method below with the struct-parameter.
 	 * 
-	 * @param Orientation Which direction content should flow and scroll
+	 * @param ScrollBarOrientation Which direction content should flow and scroll
 	 * 
 	 * @return Whether the user scroll the content this frame
 	 */
-	SLATEIM_API bool BeginScrollBox(EOrientation Orientation);
-
-	/**
-	 * Begins a container that will allow the user to scroll when its content is larger than the allotted space.
-	 *
-	 * @param Params (optional) Additional parameters to create the scroll box
-	 *
-	 * @return Whether the user scroll the content this frame
-	 */
-	SLATEIM_API bool BeginScrollBox(const FScrollBoxParams& Params = FScrollBoxParams());
-
+	SLATEIM_API bool BeginScrollBox(EOrientation ScrollBarOrientation = EOrientation::Orient_Vertical);
 	SLATEIM_API void EndScrollBox();
 
-	/**
-	 * Begin a container with a background image.
-	 *
-	 * Please use the method below with the struct-parameter.
-	 * 
-	 * @param BorderStyleName The name of the brush in FAppStyle to use as the background for this container
-	 * @param Orientation Which direction the contents of the container should flow
-	 * @param bAbsorbMouse (optional) Whether the container should handle all mouse inputs
-	 * @param ContentPadding (optional) How much to pad the contents of the container
-	 */
-	SLATEIM_API void BeginPopUp(const FName BorderStyleName, EOrientation Orientation, bool bAbsorbMouse = true, FMargin ContentPadding = FMargin(2.0f));
-	
-	/**
-	 * Begin a container with a background image.
-	 *
-	 * @param BorderStyleName The name of the brush in FAppStyle to use as the background for this container
-	 * @param Params (optional) Additional parameters to create the border
-	 */
-	SLATEIM_API void BeginPopUp(const FName BorderStyleName = TEXT("ToolPanel.GroupBorder"), const FPopUpParams& Params = FPopUpParams());
-
-	/**
-	 * Begin a container with a background image.
-	 *
-	 * Please use the method below with the struct-parameter.
-	 *
-	 * @param BorderBrush The brush to use as the background for this container
-	 * @param Orientation Which direction the contents of the container should flow
-	 * @param bAbsorbMouse (optional) Whether the container should handle all mouse inputs
-	 * @param ContentPadding (optional) How much to pad the contents of the container
-	 */
-	SLATEIM_API void BeginPopUp(const FSlateBrush* BorderBrush, EOrientation Orientation, bool bAbsorbMouse = true, FMargin ContentPadding = FMargin(2.0f));
-
-	/**
-	 * Begin a container with a background image.
-	 *
-	 * @param BorderBrush The brush to use as the background for this container
-	 * @param Params (optional) Additional params to create the border
-	 */
-	SLATEIM_API void BeginPopUp(const FSlateBrush* BorderBrush, const FPopUpParams& Params = FPopUpParams());
-	
+	SLATEIM_API void BeginPopUp(const FName BorderStyleName = TEXT("ToolPanel.GroupBorder"), EOrientation Orientation = Orient_Vertical, bool bAbsorbMouse = true, FMargin ContentPadding = FMargin(2.0f));
+	SLATEIM_API void BeginPopUp(const FSlateBrush* BorderBrush, EOrientation Orientation = Orient_Vertical, bool bAbsorbMouse = true, FMargin ContentPadding = FMargin(2.0f));
 	SLATEIM_API void EndPopUp();
 #pragma endregion Containers
 
@@ -494,7 +391,6 @@ namespace SlateIM
 	 * Set the next slot to AutoSize to its content
 	 */
 	SLATEIM_API void AutoSize();
-
 	/**
 	 * Sets the next slot to Fill the remaining space in its container
 	 *
@@ -508,20 +404,17 @@ namespace SlateIM
 	 * @param InMinWidth The minimum width for the next slot
 	 */
 	SLATEIM_API void MinWidth(float InMinWidth);
-
 	/**
 	 * Set the minimum height the next slot can have.
 	 * 
 	 * @param InMinHeight The minimum height for the next slot
 	 */
-
 	SLATEIM_API void MinHeight(float InMinHeight);
 	/**
 	 * Set the maximum width the next slot can have.
 	 * 
 	 * @param InMaxWidth The maximum width for the next slot
 	 */
-
 	SLATEIM_API void MaxWidth(float InMaxWidth);
 	/**
 	 * Set the maximum height the next slot can have.
@@ -529,7 +422,6 @@ namespace SlateIM
 	 * @param InMaxHeight The maximum height for the next slot
 	 */
 	SLATEIM_API void MaxHeight(float InMaxHeight);
-
 	/**
 	 * Set the alignments to Fill and removes autosize.
 	 */
@@ -539,19 +431,14 @@ namespace SlateIM
 #pragma region Widgets
 	/**
 	 * Display a string of text
-	 *
-	 * Please use the method below with the struct-parameter.
 	 * 
 	 * @param InText The text string to display
-	 * @param TextStyle A Text Block Style to override the default style. Changes to the style after the first invocation are not reflected,
+	 * @param TextStyle (optional) A Text Block Style to override the default style. Changes to the style after the first invocation are not reflected,
 	 *					provide a new style object to update the visuals at runtime
 	 */
-	SLATEIM_API void Text(const FStringView& InText, const FTextBlockStyle* TextStyle);
-
+	SLATEIM_API void Text(const FStringView& InText, const FTextBlockStyle* TextStyle = nullptr);
 	/**
 	 * Display a colored string of text
-	 *
-	 * Please use the method below with the struct-parameter.
 	 * 
 	 * @param InText The text string to display
 	 * @param Color The color to display the text as
@@ -560,14 +447,6 @@ namespace SlateIM
 	 */
 	SLATEIM_API void Text(const FStringView& InText, FSlateColor Color, const FTextBlockStyle* TextStyle = nullptr);
 
-	/**
-	 * Display a string of text
-	 *
-	 * @param InText The text string to display
-	 * @param Params (optional) Additional parameters to create the text.
-	 */
-	SLATEIM_API void Text(const FStringView& InText, const FTextParams& Params = FTextParams());
-
 	/**
 	 * Create a text input field
 	 * 
@@ -578,185 +457,96 @@ namespace SlateIM
 	 *					
 	 * @return Whether the user changed or committed text this frame
 	 */
-	SLATEIM_API bool EditableText(FString& InOutText, const FStringView& HintText, const FEditableTextBoxStyle* TextStyle = nullptr);
-
-	/**
-	 * Create a text input field
-	 *
-	 * @param InOutText The text that is in the input field. Can have a default value
-	 * @param Params (optional) Additional parameters to create the widget
-	 *
-	 * @return Whether the user changed or committed text this frame
-	 */
-	SLATEIM_API bool EditableText(FString& InOutText, const FEditableTextParams& Params = FEditableTextParams());
+	SLATEIM_API bool EditableText(FString& InOutText, const FStringView& HintText = FStringView(), const FEditableTextBoxStyle* TextStyle = nullptr);
 
 	/**
 	 * Display a brush
 	 * 
 	 * @param ImageBrush The brush to display
-	 * @param ColorAndOpacity The color to tint the brush
+	 * @param ColorAndOpacity (optional) The color to tint the brush
 	 * @param DesiredSize (optional) Override the desired size to display the brush at
 	 */
-	SLATEIM_API void Image(const FSlateBrush* ImageBrush, const FSlateColor& ColorAndOpacity, FVector2D DesiredSize = FVector2D::ZeroVector);
-
-	/**
-	 * Display a brush
-	 *
-	 * @param ImageBrush The brush to display
-	 * @param Params (optional) Additional parameters to display the image
-	 */
-	SLATEIM_API void Image(const FSlateBrush* ImageBrush, const FImageParams& Params = FImageParams());
-
+	SLATEIM_API void Image(const FSlateBrush* ImageBrush, const FSlateColor& ColorAndOpacity = FLinearColor::White, FVector2D DesiredSize = FVector2D::ZeroVector);
 	/**
 	 * Display a slate style brush
-	 *
-	 * Please use the method below with the struct-parameter.
 	 * 
 	 * @param ImageStyleName The named of the brush to display from FAppStyle
-	 * @param ColorAndOpacity The color to tint the brush
+	 * @param ColorAndOpacity (optional) The color to tint the brush
 	 * @param DesiredSize (optional) Override the desired size to display the brush at
 	 */
-	SLATEIM_API void Image(const FName ImageStyleName, const FSlateColor& ColorAndOpacity, FVector2D DesiredSize = FVector2D::ZeroVector);
-	
-	/**
-	 * Display a slate style brush
-	 *
-	 * @param ImageStyleName The named of the brush to display from FAppStyle
-	 * @param Params (optional) Additional parameters to display the image
-	 */
-	SLATEIM_API void Image(const FName ImageStyleName, const FImageParams& Params = FImageParams());
-
+	SLATEIM_API void Image(const FName ImageStyleName, const FSlateColor& ColorAndOpacity = FLinearColor::White, FVector2D DesiredSize = FVector2D::ZeroVector);
 	/**
 	 * Display a colored box
-	 *
-	 * Please use the method below with the struct-parameter.
 	 * 
 	 * @param ColorAndOpacity The color to display
 	 * @param DesiredSize (optional) Override the desired size to display the box at
 	 */
 	SLATEIM_API void Image(const FSlateColor& ColorAndOpacity, FVector2D DesiredSize = FVector2D::ZeroVector);
-
-	/**
-	 * Display a colored box
-	 *
-	 * @param Params (optional) Additional parameters to display the block (including the color)
-	 */
-	SLATEIM_API void Image(const FImageParams& Params = FImageParams());
-
 #if WITH_ENGINE
 	/**
 	 * Display a texture
-	 *
-	 * Please use the method below with the struct-parameter.
 	 * 
 	 * @param ImageTexture The texture to display
 	 * @param ColorAndOpacity (optional) The color to tint the brush
 	 * @param DesiredSize (optional) Override the desired size to display the brush at
 	 */
-	
-	SLATEIM_API void Image(UTexture2D* ImageTexture, const FSlateColor& ColorAndOpacity, FVector2D DesiredSize = FVector2D::ZeroVector);
-	
-	/**
-	 * Display a texture
-	 *
-	 * @param ImageTexture The texture to display
-	 * @param Params (optional) Additional parameters to display the image
-	 */
-	SLATEIM_API void Image(UTexture2D* ImageTexture, const FImageParams& Params = FImageParams());
-
+	SLATEIM_API void Image(UTexture2D* ImageTexture, const FSlateColor& ColorAndOpacity = FLinearColor::White, FVector2D DesiredSize = FVector2D::ZeroVector);
 	/**
 	 * Display a render target texture
-	 *
-	 * Please use the method below with the struct-parameter.
 	 * 
 	 * @param ImageRenderTarget The render target texture to display
-	 * @param ColorAndOpacity The color to tint the brush
+	 * @param ColorAndOpacity (optional) The color to tint the brush
 	 * @param DesiredSize (optional) Override the desired size to display the brush at
 	 */
-	
-	SLATEIM_API void Image(UTextureRenderTarget2D* ImageRenderTarget, const FSlateColor& ColorAndOpacity, FVector2D DesiredSize = FVector2D::ZeroVector);
-	
-	/**
-	 * Display a render target texture
-	 *
-	 * @param ImageRenderTarget The render target texture to display
-	 * @param Params (optional) Additional parameters to display the image
-	 */
-	SLATEIM_API void Image(UTextureRenderTarget2D* ImageRenderTarget, const FImageParams& Params = FImageParams());
-
+	SLATEIM_API void Image(UTextureRenderTarget2D* ImageRenderTarget, const FSlateColor& ColorAndOpacity = FLinearColor::White, FVector2D DesiredSize = FVector2D::ZeroVector);
 	/**
 	 * Display a material
-	 *
-	 * Please use the method below with the struct-parameter.
 	 * 
 	 * @param ImageMaterial The material to display
 	 * @param BrushSize The size to create the internal brush with
-	 * @param ColorAndOpacity The color to tint the brush
+	 * @param ColorAndOpacity (optional) The color to tint the brush
 	 * @param DesiredSize (optional) Override the desired size to display the brush at
 	 */
-	SLATEIM_API void Image(UMaterialInterface* ImageMaterial, FVector2D BrushSize, const FSlateColor& ColorAndOpacity, FVector2D DesiredSize = FVector2D::ZeroVector);
-
-	/**
-	 * Display a material
-	 *
-	 * @param ImageMaterial The material to display
-	 * @param BrushSize The size to create the internal brush with
-	 * @param Params (optional) Additional parameters to display the image
-	 */
-	SLATEIM_API void Image(UMaterialInterface* ImageMaterial, FVector2D BrushSize, const FImageParams& Params = FImageParams());
+	SLATEIM_API void Image(UMaterialInterface* ImageMaterial, FVector2D BrushSize, const FSlateColor& ColorAndOpacity = FLinearColor::White, FVector2D DesiredSize = FVector2D::ZeroVector);
 #endif
 
 	/**
 	 * Display a button with text
 	 * 
-	 * Please use the method below with the struct-parameter.
-
 	 * @param InText The text to display on the button
-	 * @param bEnabled Indicate if the button is enabled or not (grayed)
 	 * @param InStyle (optional) A Button Style to override the default style. Changes to the style after the first invocation are not reflected,
 	 *				  provide a new style object to update the visuals at runtime
 	 *					
 	 * @return Whether the user clicked the button this frame
 	 */
-	SLATEIM_API bool Button(const FStringView& InText, const bool bEnabled, const FButtonStyle* InStyle = nullptr);
+	SLATEIM_API bool Button(const FStringView& InText, const FButtonStyle* InStyle = nullptr);
 
 	/**
-	 *
+	 * Display a button with text
+	 * 
 	 * @param InText The text to display on the button
-	 * @param Params (optional) Additional parameters to create the button
-	 *
+	 * @param bEnabled Indicate if the button is enabled or not (grayed)
+	 * @param InStyle (optional) A Button Style to override the default style. Changes to the style after the first invocation are not reflected,
+	 *				  provide a new style object to update the visuals at runtime
+	 *					
 	 * @return Whether the user clicked the button this frame
 	 */
-	SLATEIM_API bool Button(const FStringView& InText, const FButtonParams& Params = FButtonParams());
+	SLATEIM_API bool Button(const FStringView& InText, const bool bEnabled, const FButtonStyle* InStyle = nullptr);
 
 	/**
 	 * Display a two-state checkbox
-	 *
-	 * Please use the method below with the struct-parameter.
 	 * 
-	 * @param InLabel The label for the checkbox
+	 * @param InText The label for the checkbox
 	 * @param InOutCurrentState The current state of the checkbox. Can provide a default value.
 	 * @param CheckBoxStyle (optional) A Checkbox Style to override the default style. Changes to the style after the first invocation are not reflected,
 	 *						provide a new style object to update the visuals at runtime
 	 *					
 	 * @return Whether the user changed the state of the checkbox this frame
 	 */
-	SLATEIM_API bool CheckBox(const FStringView& InLabel, bool& InOutCurrentState, const FCheckBoxStyle* CheckBoxStyle = nullptr);
-
-	/**
-	 * Display a two-state checkbox
-	 *
-	 * @param InOutCurrentState The current state of the checkbox. Can provide a default value.
-	 * @param Params (optional) Additional parameters to create the checkbox
-	 *
-	 * @return Whether the user changed the state of the checkbox this frame
-	 */
-	SLATEIM_API bool CheckBox(bool& InOutCurrentState, const FCheckBoxParams& Params = FCheckBoxParams());
+	SLATEIM_API bool CheckBox(const FStringView& InText, bool& InOutCurrentState, const FCheckBoxStyle* CheckBoxStyle = nullptr);
 	
 	/**
 	 * Display a three-state checkbox
-	 *
-	 * Please use the method below with the struct-parameter.
 	 * 
 	 * @param InText The label for the checkbox
 	 * @param InOutCurrentState The current state of the checkbox. Can provide a default value.
@@ -767,20 +557,8 @@ namespace SlateIM
 	 */
 	SLATEIM_API bool CheckBox(const FStringView& InText, ECheckBoxState& InOutCurrentState, const FCheckBoxStyle* CheckBoxStyle = nullptr);
 
-	/**
-	 * Display a three-state checkbox
-	 *
-	 * @param InOutCurrentState The current state of the checkbox. Can provide a default value.
-	 * @param Params (optional) Additional parameters to create the checkbox
-	 *
-	 * @return Whether the user changed the state of the checkbox this frame
-	 */
-	SLATEIM_API bool CheckBox(ECheckBoxState& InOutCurrentState, const FCheckBoxParams& Params = FCheckBoxParams());
-
 	/**
 	 * Display a float-based spin box
-	 *
-	 * Please use the method below with the struct-parameter.
 	 * 
 	 * @param InOutValue The current value of the spin box. Can provide a default value.
 	 * @param Min The minimum value of the spin box (or unset for no limit).
@@ -791,11 +569,8 @@ namespace SlateIM
 	 * @return Whether the user changed or committed the spin box value this frame
 	 */
 	SLATEIM_API bool SpinBox(float& InOutValue, TOptional<float> Min, TOptional<float> Max, const FSpinBoxStyle* SpinBoxStyle = nullptr);
-
 	/**
 	 * Display a double-based spin box
-	 *
-	 * Please use the method below with the struct-parameter.
 	 * 
 	 * @param InOutValue The current value of the spin box. Can provide a default value.
 	 * @param Min The minimum value of the spin box (or unset for no limit).
@@ -806,11 +581,8 @@ namespace SlateIM
 	 * @return Whether the user changed or committed the spin box value this frame
 	 */
 	SLATEIM_API bool SpinBox(double& InOutValue, TOptional<double> Min, TOptional<double> Max, const FSpinBoxStyle* SpinBoxStyle = nullptr);
-
 	/**
 	 * Display an integer-based spin box
-	 *
-	 * Please use the method below with the struct-parameter.
 	 * 
 	 * @param InOutValue The current value of the spin box. Can provide a default value.
 	 * @param Min The minimum value of the spin box (or unset for no limit).
@@ -822,97 +594,31 @@ namespace SlateIM
 	 */
 	SLATEIM_API bool SpinBox(int32& InOutValue, TOptional<int32> Min, TOptional<int32> Max, const FSpinBoxStyle* SpinBoxStyle = nullptr);
 
-	/**
-	 * Display a spin box
-	 *
-	 * @param InOutValue The current value of the spin box. Can provide a default value.
-	 * @param Params (optional) Additional parameters to create the spinbox
-	 *
-	 * @return Whether the user committed the spin box value this frame
-	 */
-	template<typename InNumericType>
-	bool SpinBox(InNumericType& InOutValue, const FSpinBoxParams<InNumericType>& Params = FSpinBoxParams<InNumericType>());
-
-	/**
-	 * Display a float-based spin box
-	 *
-	 * @param InOutValue The current value of the spin box. Can provide a default value.
-	 * @param Params Additional parameters to create the spinbox
-	 * 
-	 * @return Whether the user changed or committed the spin box value this frame
-	 */
-	template<>
-	bool SLATEIM_API SpinBox(float& InOutValue, const FSpinBoxParams<float>& Params);
-
-	/**
-	 * Display a double-based spin box
-	 *
-	 * @param InOutValue The current value of the spin box. Can provide a default value.
-	 * @param Params Additional parameters to create the spinbox
-	 *
-	 * @return Whether the user changed or committed the spin box value this frame
-	 */
-	template<>
-	bool SLATEIM_API SpinBox(double& InOutValue, const FSpinBoxParams<double>& Params);
-
-	/**
-	 * Display an int32-based spin box
-	 *
-	 * @param InOutValue The current value of the spin box. Can provide a default value.
-	 * @param Params Additional parameters to create the spinbox
-	 * 
-	 * @return Whether the user changed or committed the spin box value this frame
-	 */
-	template<>
-	bool SLATEIM_API SpinBox(int32& InOutValue, const FSpinBoxParams<int32>& Params);
-
 	/**
 	 * Display a float-based slider
-	 *
-	 * Please use the method below with the struct-parameter.
-	 *
+	 * 
 	 * @param InOutValue The current value of the slider. Can provide a default value.
 	 * @param Min The minimum value of the slider
 	 * @param Max The maximum value of the slider
 	 * @param Step The smallest incremental change that can be made to the slider's value
 	 * @param SliderStyle (optional) A Slider Style to override the default style. Changes to the style after the first invocation are not reflected,
 	 *					  provide a new style object to update the visuals at runtime
-	 *
+	 *					  
 	 * @return Whether the user changed the value of the slider this frame
 	 */
 	SLATEIM_API bool Slider(float& InOutValue, float Min, float Max, float Step, const FSliderStyle* SliderStyle = nullptr);
 
-	/**
-	 * Display a float-based slider
-	 *
-	 * @param InOutValue The current value of the slider. Can provide a default value.
-	 * @param Params (optional) Additional parameters to create the slider
-	 *
-	 * @return Whether the user changed the value of the slider this frame
-	 */
-	SLATEIM_API bool Slider(float& InOutValue, const FSliderParams& Params = FSliderParams());
-
 	/**
 	 * Display a progress bar
 	 * 
 	 * @param Percent The value of the progress bar (0.0 to 1.0)
-	 * @param ProgressBarStyle A Progress Bar Style to override the default style. Changes to the style after the first invocation are not reflected,
+	 * @param ProgressBarStyle (optional) A Progress Bar Style to override the default style. Changes to the style after the first invocation are not reflected,
 	 *						   provide a new style object to update the visuals at runtime
 	 */
-	SLATEIM_API void ProgressBar(TOptional<float> Percent, const FProgressBarStyle* ProgressBarStyle);
-
-	/**
-	 * Display a progress bar
-	 *
-	 * @param Percent The value of the progress bar (0.0 to 1.0)
-	 * @param Params (optional) Additional parameters to create the progress bar
-	 */
-	SLATEIM_API void ProgressBar(TOptional<float> Percent, const FProgressBarParams& Params = FProgressBarParams());
+	SLATEIM_API void ProgressBar(TOptional<float> Percent, const FProgressBarStyle* ProgressBarStyle = nullptr);
 
 	/**
 	 * Display a dropdown of text options
-	 *
-	 * Please use the method below with the struct-parameter.
 	 * 
 	 * @param ComboItems The options the user can choose from
 	 * @param InOutSelectedItemIndex The index of the currently selected option
@@ -922,12 +628,10 @@ namespace SlateIM
 	 *					   
 	 * @return Whether the user changed the selected option this frame
 	 */
-	SLATEIM_API bool ComboBox(const TArray<FString>& ComboItems, int32& InOutSelectedItemIndex, bool bForceRefresh, const FComboBoxStyle* InComboStyle = nullptr);
+	SLATEIM_API bool ComboBox(const TArray<FString>& ComboItems, int32& InOutSelectedItemIndex, bool bForceRefresh = false, const FComboBoxStyle* InComboStyle = nullptr);
 
 	/**
 	 * Display a searchable dropdown of text options
-	 *
-	 * Please use the method below with the struct-parameter.
 	 * 
 	 * @param ComboItems The options the user can choose from
 	 * @param InOutSelectedItemIndex The index of the currently selected option
@@ -938,21 +642,8 @@ namespace SlateIM
 	 */
 	SLATEIM_API bool SearchableComboBox(const TArray<FString>& ComboItems, int32& InOutSelectedItemIndex, bool bForceRefresh = false, const FComboBoxStyle* InComboStyle = nullptr);
 
-	/**
-	 * Display a dropdown of text options
-	 *
-	 * @param ComboItems The options the user can choose from
-	 * @param InOutSelectedItemIndex The index of the currently selected option
-	 * @param Params (optional) Additional parameter to create the combo box
-	 *
-	 * @return Whether the user changed the selected option this frame
-	 */
-	SLATEIM_API bool ComboBox(const TArray<FString>& ComboItems, int32& InOutSelectedItemIndex, const FComboBoxParams& Params = FComboBoxParams());
-
 	/**
 	 * Display a list of text options
-	 *
-	 * Please use the method below with the struct-parameter.
 	 * 
 	 * @param ListItems The list of options the user can choose from
 	 * @param InOutSelectedItemIndex The index of the currently selected option
@@ -962,18 +653,7 @@ namespace SlateIM
 	 *					   
 	 * @return Whether the user changed the selected option this frame
 	 */
-	SLATEIM_API bool SelectionList(const TArray<FString>& ListItems, int32& InOutSelectedItemIndex, bool bForceRefresh, const FTableViewStyle* InStyle = nullptr);
-
-	/**
-	 * Display a list of text options
-	 *
-	 * @param ListItems The list of options the user can choose from
-	 * @param InOutSelectedItemIndex The index of the currently selected option
-	 * @param Params (optional) Additional parameters to create the selection list
-	 *
-	 * @return Whether the user changed the selected option this frame
-	 */
-	SLATEIM_API bool SelectionList(const TArray<FString>& ListItems, int32& InOutSelectedItemIndex, const FSelectionListParams& Params = FSelectionListParams());
+	SLATEIM_API bool SelectionList(const TArray<FString>& ListItems, int32& InOutSelectedItemIndex, bool bForceRefresh = false, const FTableViewStyle* InStyle = nullptr);
 
 	/**
 	 * Create a block of empty space
@@ -982,11 +662,10 @@ namespace SlateIM
 	 */
 	SLATEIM_API void Spacer(const FVector2D& Size);
 
-	/**
+	/** 
 	 * Add an already created slate widget to the IM hierarchy
-	 *
+	 * 
 	 * @note Do not create the widget you are passing in here every frame. This is meant to be used to add already created widgets
-	 * @note Not available in BPFL: TSharedRef<SWidget> is not Blueprint-compatible
 	 */
 	SLATEIM_API void Widget(TSharedRef<SWidget> InWidget);
 #pragma endregion Widgets
@@ -998,7 +677,6 @@ namespace SlateIM
 	 * @return Whether the menu is open
 	 */
 	SLATEIM_API bool BeginContextMenuAnchor();
-
 	SLATEIM_API void EndContextMenuAnchor();
 
 	/** 
@@ -1037,73 +715,33 @@ namespace SlateIM
 
 	/**
 	 * Add a menu item button to the current context menu
-	 *
-	 * Please use the method below with the struct-parameter.
 	 * 
 	 * @param RowText Label for the menu item
-	 * @param ToolTipText The tooltip to display for the menu item
+	 * @param ToolTipText (optional) The tooltip to display for the menu item
 	 * 
 	 * @return Whether the user clicked the menu item this frame
 	 */
-	SLATEIM_API bool AddMenuButton(const FStringView& RowText, const FStringView& ToolTipText);
-
-	/**
-	 * Add a menu item button to the current context menu
-	 *
-	 * @param RowText Label for the menu item
-	 * @param Params (optional) Additional parameters to create the menu button
-	 *
-	 * @return Whether the user clicked the menu item this frame
-	 */
-	SLATEIM_API bool AddMenuButton(const FStringView& RowText, const FMenuButtonParams& Params = FMenuButtonParams());
-
+	SLATEIM_API bool AddMenuButton(const FStringView& RowText, const FStringView& ToolTipText = FStringView());
 	/**
 	 * Display a menu item button with a checkbox
-	 *
-	 * Please use the method below with the struct-parameter.
 	 * 
 	 * @param RowText Label for the menu item
 	 * @param InOutCurrentState The current state of the toggle
-	 * @param ToolTipText The tooltip to display for the menu item
+	 * @param ToolTipText (optional) The tooltip to display for the menu item
 	 * 
 	 * @return Whether the user toggled the checkbox this frame
 	 */
-	SLATEIM_API bool AddMenuToggleButton(const FStringView& RowText, bool& InOutCurrentState, const FStringView& ToolTipText);
-
-	/**
-	 * Display a menu item button with a checkbox
-	 *
-	 * @param RowText Label for the menu item
-	 * @param InOutCurrentState The current state of the toggle
-	 * @param Params (optional) Additional parameters to create the menu button
-	 *
-	 * @return Whether the user toggled the checkbox this frame
-	 */
-	SLATEIM_API bool AddMenuToggleButton(const FStringView& RowText, bool& InOutCurrentState, const FMenuButtonParams& Params = FMenuButtonParams());
-
+	SLATEIM_API bool AddMenuToggleButton(const FStringView& RowText, bool& InOutCurrentState, const FStringView& ToolTipText = FStringView());
 	/**
 	 * Display a menu item button with a checkmark
-	 *
-	 * Please use the method below with the struct-parameter.
 	 * 
 	 * @param RowText Label for the menu item
 	 * @param InOutCurrentState The current state of the checkmark
-	 * @param ToolTipText The tooltip to display for the menu item
+	 * @param ToolTipText (optional) The tooltip to display for the menu item
 	 * 
 	 * @return Whether the user clicked the menu item this frame
 	 */
-	SLATEIM_API bool AddMenuCheckButton(const FStringView& RowText, bool& InOutCurrentState, const FStringView& ToolTipText);
-
-	/**
-	 * Display a menu item button with a checkmark
-	 *
-	 * @param RowText Label for the menu item
-	 * @param InOutCurrentState The current state of the checkmark
-	 * @param Params (optional) Additional parameters to create the menu button
-	 *
-	 * @return Whether the user clicked the menu item this frame
-	 */
-	SLATEIM_API bool AddMenuCheckButton(const FStringView& RowText, bool& InOutCurrentState, const FMenuButtonParams& Params = FMenuButtonParams());
+	SLATEIM_API bool AddMenuCheckButton(const FStringView& RowText, bool& InOutCurrentState, const FStringView& ToolTipText = FStringView());
 
 	/**
 	 * Add a submenu item to the current menu
@@ -1111,7 +749,6 @@ namespace SlateIM
 	 * @param SubMenuText The label for the menu item
 	 */
 	SLATEIM_API void BeginSubMenu(const FStringView& SubMenuText);
-
 	SLATEIM_API void EndSubMenu();
 #pragma endregion Menu
 
@@ -1159,26 +796,14 @@ namespace SlateIM
 
 	/**
 	 * Begins a Tab to contain any other content
-	 *
-	 * Please use the method below with the struct-parameter.
 	 * 
 	 * @param TabId Provide a tab id that is unique within the Tab Group
-	 * @param TabIcon An icon to display on the tab next to the title. This cannot be updated after the initial creation of the tab.
+	 * @param TabIcon (optional) An icon to display on the tab next to the title. This cannot be updated after the initial creation of the tab.
 	 * @param TabTitle (optional) The title to display for the tab if it should differ from the TabId. This cannot be updated after the initial creation of the tab.
 	 * 
 	 * @return Returns true if the tab is active, false otherwise. Logic to draw the contents of the tab can be skipped when this is false.
 	 */
-	SLATEIM_API bool BeginTab(const FName& TabId, const FSlateIcon& TabIcon, const FText& TabTitle = FText::GetEmpty());
-
-	/**
-	 * Begins a Tab to contain any other content
-	 *
-	 * @param TabId Provide a tab id that is unique within the Tab Group
-	 * @param Params (optional) Additional parameters used to create the tab
-	 *
-	 * @return Returns true if the tab is active, false otherwise. Logic to draw the contents of the tab can be skipped when this is false.
-	 */
-	SLATEIM_API bool BeginTab(const FName& TabId, const FTabParams& Params = FTabParams());
+	SLATEIM_API bool BeginTab(const FName& TabId, const FSlateIcon& TabIcon = FSlateIcon(), const FText& TabTitle = FText::GetEmpty());
 
 	/**
 	 * Ends the current Tab. This must be called even when BeginTab returns false.
@@ -1211,7 +836,6 @@ namespace SlateIM
 	 * @note It is not possible to enable a widget inside a disabled parent widget by calling EndDisabledState before a child is created inside a disabled Widget.
 	 */
 	SLATEIM_API void BeginDisabledState();
-
 	SLATEIM_API void EndDisabledState();
 
 	/**
@@ -1223,73 +847,36 @@ namespace SlateIM
 
 	/**
 	 * Opens a modal dialog of the specified type. This function will not return until the user closes the modal dialog
-	 *
-	 * Please use the method below with the struct-parameter.
 	 * 
 	 * @param MessageType The type of options the user can respond to the dialog with
 	 * @param DialogText The message to display to the user
 	 * @param Category The type of dialog to display
-	 * @param DialogTitle (optional) The title to display in the dialog window
+	 * @param DialogTitle The title to display in the dialog window
 	 * 
 	 * @return The option the user selected
 	 *
 	 * @note SlateIM updates are disable while a SlateIM modal is open
 	 * @see CanUpdateSlateIM()
 	 */
-	SLATEIM_API EAppReturnType::Type ModalDialog(EAppMsgType::Type MessageType, const FStringView& DialogText, EAppMsgCategory Category, const FStringView& DialogTitle = FStringView());
-
-	/**
-	 * Opens a modal dialog of the specified type. This function will not return until the user closes the modal dialog
-	 *
-	 * @param MessageType The type of options the user can respond to the dialog with
-	 * @param DialogText The message to display to the user
-	 * @param Params (optional) Additional parameters used to create the dialog
-	 *
-	 * @return The option the user selected
-	 *
-	 * @note SlateIM updates are disable while a SlateIM modal is open
-	 * @see CanUpdateSlateIM()
-	 */
-	SLATEIM_API EAppReturnType::Type ModalDialog(EAppMsgType::Type MessageType, const FStringView& DialogText, const FModalDialogParams& Params = FModalDialogParams());
+	SLATEIM_API EAppReturnType::Type ModalDialog(EAppMsgType::Type MessageType, const FStringView& DialogText, EAppMsgCategory Category = EAppMsgCategory::Warning, const FStringView& DialogTitle = FStringView());
 
 	/**
 	 * Text filter input box that updates the filter when text changes.
 	 *
 	 * @param InFilter The text filter to display and update
-	 * @param Params (optional) Additional parameters to create the text filter
-	 *
+	 * @param HintText Optional hint text to display when the input is empty
 	 * @return Whether the user changed or committed text this frame
-	 *
-	 * @note Not available in BPFL: template function
 	 */
 	template<typename ItemType>
-	bool TextFilter(FSlateImTextFilter<ItemType>& InFilter, const FEditableTextParams& Params = FEditableTextParams())
+	bool TextFilter(FSlateImTextFilter<ItemType>& InFilter, const FStringView& HintText = FStringView())
 	{
-		if (EditableText(InFilter.InputText, Params) && InFilter.TextFilter.IsValid())
+		if (EditableText(InFilter.InputText, HintText) && InFilter.TextFilter.IsValid())
 		{
 			InFilter.TextFilter->SetRawFilterText(FText::FromString(InFilter.InputText));
 			return true;
 		}
 		return false;
 	}
-
-	/**
-	 * Text filter input box that updates the filter when text changes.
-	 *
-	 * Please use the method above with the struct-parameter.
-	 *
-	 * @param InFilter The text filter to display and update
-	 * @param HintText hint text to display when the input is empty
-	 * 
-	 * @return Whether the user changed or committed text this frame
-	 *
-	 * @note Not available in BPFL: template function
-	 */
-	template<typename ItemType>
-	bool TextFilter(FSlateImTextFilter<ItemType>& InFilter, const FStringView& HintText)
-	{
-		return TextFilter<ItemType>(InFilter, {.HintText = HintText});
-	}
 #pragma endregion Util
 
 #pragma region Queries
@@ -1300,6 +887,15 @@ namespace SlateIM
 	 */
 	SLATEIM_API bool IsHovered();
 
+	enum class EFocusDepth : uint8
+	{
+		// Only check if the previous widget is focused
+		SelfOnly,
+
+		// Check if the previous widget or any of its childs widgets are focused (recursively)
+		IncludingDescendants
+	};
+
 	/**
 	 * Query whether the previous widget is focused
 	 * 
@@ -1322,8 +918,6 @@ namespace SlateIM
 
 	/**
 	 * Add a line graph of 2D vectors to the graph widget
-	 *
-	 * Please use the method below with the struct-parameter.
 	 * 
 	 * @param Points The X,Y points to plot on the graph
 	 * @param LineColor The color of the line to draw for this graph
@@ -1333,19 +927,9 @@ namespace SlateIM
 	 */
 	SLATEIM_API void GraphLine(const TArrayView<FVector2D>& Points, const FLinearColor& LineColor, float LineThickness, const FDoubleRange& XViewRange, const FDoubleRange& YViewRange);
 
-	/**
-	 * Add a line graph of 2D vectors to the graph widget
-	 *
-	 * @param Points The X,Y points to plot on the graph
-	 * @param Params Additional parameters to draw the line
-	 */
-	SLATEIM_API void GraphLine(const TArrayView<FVector2D>& Points, const FGraphLinePointsParams& Params = FGraphLinePointsParams());
-
 	/**
 	 * Add a line graph of values to the graph widget.
 	 * The value index is used as the X-value when plotting, the graph will scale horizontally to fit all values in the array.
-	 *
-	 * Please use the method below with the struct-parameter.
 	 * 
 	 * @param Values The Y values to plot on the graph
 	 * @param LineColor The color of the line to draw for this graph
@@ -1353,15 +937,6 @@ namespace SlateIM
 	 * @param ViewRange The min and max Y values to vertically scale the graph to
 	 */
 	SLATEIM_API void GraphLine(const TArrayView<double>& Values, const FLinearColor& LineColor, float LineThickness, const FDoubleRange& ViewRange);
-
-	/**
-	 * Add a line graph of values to the graph widget.
-	 * The value index is used as the X-value when plotting, the graph will scale horizontally to fit all values in the array.
-	 *
-	 * @param Values The Y values to plot on the graph
-	 * @param Params Additional parameters to draw the line
-	 */
-	SLATEIM_API void GraphLine(const TArrayView<double>& Values, const FGraphLineValuesParams& Params = FGraphLineValuesParams());
 #pragma endregion Graphing
 
 #pragma region Inputs
@@ -1410,41 +985,64 @@ namespace SlateIM
 #if WITH_ENGINE
 	namespace Canvas
 	{
+		enum class ECanvasUpdateType : uint8
+		{
+			/** Redraw the canvas every frame. */
+			EveryFrame,
+
+			/**
+			 * Only redraw when specifically invalidated.
+			 * Warning: Textures and materials may not be rendered correctly if they are not loaded or shaders are not cached.
+			 */
+			Invalidation
+
+			/** TODO: Automatically redraw when the hash of the draw queue changes. */
+			// HashChange
+		};
+
+		struct FCanvasParams
+		{
+			FSlateColor ColorAndOpacity = FSlateColor(EStyleColor::White);
+			FVector2D DesiredSize = FVector2D::ZeroVector;
+			UObject* WorldContext = nullptr;
+			ECanvasUpdateType UpdateType = ECanvasUpdateType::EveryFrame;
+		};
+
 		/**
 		 * Begin queueing commands to draw to a canvas render target.
 		 * 
 		 * @param Width The width of the canvas in pixels.
 		 * @param Height The height of the canvas in pixels.
-		 * @param Params Additional parameters to create the canvas.
+		 * @param ColorAndOpacity (optional) Tint applied to the resulting image.
+		 * @param DesiredSize (optional) The size of the image drawn in the UI. This may be different from the size of the canvas.
+		 * @param WorldContext (optional) The world related to the canvas being drawn, if any.
 		 */
-		SLATEIM_API void BeginCanvas(int32 Width, int32 Height, const FCanvasParams& CanvasParams = FCanvasParams());
+		void BeginCanvas(int32 Width, int32 Height, const FCanvasParams& CanvasParams = FCanvasParams());
 
 		/**
 		 * Ends the currently drawn canvas and executes all queued commands.
 		 */
-		SLATEIM_API void EndCanvas();
+		void EndCanvas();
 
 		/**
 		 * Invalidate the current canvas, forcing a redraw if its UpdateType is set to Invalidation.
 		 */
-		SLATEIM_API void Invalidate();
+		void Invalidate();
 
 		/** Sets the position of the lower-right corner of the clipping region of the canvas */
-		SLATEIM_API void SetClip(const FVector2f& ClipPosition);
+		void SetClip(const FVector2f& ClipPosition);
 
 		/** Sets the draw color of the canvas. */
-		SLATEIM_API void SetDrawColor(const FLinearColor& Color);
-		SLATEIM_API void SetDrawColor(const FColor& Color);
+		void SetDrawColor(const FLinearColor& Color);
+		void SetDrawColor(const FColor& Color);
 
 		/**
 		 * Queue drawing a canvas item to the canvas.
-		 *
+		 * 
 		 * @param Item The canvas item to draw. The item *must* stay in scope until EndCanvas().
 		 * @param CanvasPosition (optional) The position to draw the item on the canvas.
-		 *
-		 * @note Not available in BPFL: TSharedRef<FCanvasItem> is not Blueprint-compatible
 		 */
-		SLATEIM_API void DrawItem(TSharedRef<FCanvasItem> ItemPtr, const FVector2f& CanvasPosition = FVector2f::ZeroVector);
+		void DrawItem(TSharedRef<FCanvasItem> ItemPtr, const FVector2f& CanvasPosition = FVector2f::ZeroVector);
 
 		/**
 		 * Queue drawing a canvas icon (texture) to the canvas.
@@ -1453,7 +1051,7 @@ namespace SlateIM
 		 * @param CanvasPosition (optional) The position of the icon on the canvas, in pixels.
 		 * @param Scale (optional) The scale of the icon.
 		 */
-		SLATEIM_API void DrawIcon(const FCanvasIcon& Icon, const FVector2f& CanvasPosition = FVector2f::ZeroVector, const FVector2f& Scale = FVector2f::UnitVector);
+		void DrawIcon(const FCanvasIcon& Icon, const FVector2f& CanvasPosition = FVector2f::ZeroVector, const FVector2f& Scale = FVector2f::UnitVector);
 
 		/**
 		 * Queue drawing a line to the canvas.
@@ -1463,7 +1061,31 @@ namespace SlateIM
 		 * @param Thickness (optional) The thickness of the line in pixels.
 		 * @param RenderColor (optional) The color of the line.
 		 */
-		SLATEIM_API void DrawLine(const FVector2D& CanvasPositionA, const FVector2D& CanvasPositionB, float Thickness = 1.0f, const FLinearColor& RenderColor = FLinearColor::White);
+		void DrawLine(const FVector2D& CanvasPositionA, const FVector2D& CanvasPositionB, float Thickness = 1.0f, FLinearColor RenderColor = FLinearColor::White);
+
+		/** 
+		 * Extra optional parameters for rendering tiles (textures and materials.)
+		 */
+		struct FTileRenderParams
+		{
+			// The top left of the area on the tile to draw, in texels.
+			FVector2D UVPosition = FVector2D::ZeroVector;
+
+			// The size of the area on the tile to draw, in texels.
+			FVector2D UVSize = FVector2D::UnitVector;
+
+			// A tint to apply to tile.
+			FLinearColor RenderColor = FLinearColor::White;
+
+			// The blend mode with which to draw the tile.
+			EBlendMode BlendMode = BLEND_Translucent;
+
+			// The rotation of the tile, in degrees.
+			float Rotation = 0.f;
+
+			// The point on the tile around which it is rotated, in texels.
+			FVector2D PivotPoint = FVector2D(0.5, 0.5);
+		};
 
 		/**
 		 * Queue drawing a texture to the canvas.
@@ -1471,9 +1093,9 @@ namespace SlateIM
 		 * @param RenderTexture The texture to draw.
 		 * @param CanvasPosition The position of the texture on the canvas, in pixels.
 		 * @param CanvasSize The size of the texture on the canvas, in pixels.
-		 * @param RenderParams (optional) Additional parameters to draw the tile.
+		 * @param RenderParams (optional) Extra optional parameters for tile drawing.
 		 */
-		SLATEIM_API void DrawTexture(UTexture* RenderTexture, const FVector2D& CanvasPosition, const FVector2D& CanvasSize, const FTileRenderParams& RenderParams = FTileRenderParams());
+		void DrawTexture(UTexture* RenderTexture, const FVector2D& CanvasPosition, const FVector2D& CanvasSize, const FTileRenderParams& RenderParams = FTileRenderParams());
 
 		/**
 		 * Queue drawing a material to the canvas.
@@ -1481,9 +1103,42 @@ namespace SlateIM
 		 * @param RenderMaterial The material to draw.
 		 * @param CanvasPosition The position of the material on the canvas, in pixels.
 		 * @param CanvasSize The size of the material on the canvas, in pixels.
-		 * @param RenderParams (optional) Additional parameters to draw the tile. RenderColor and BlendMode are not used.
+		 * @param RenderParams (optional) Extra optional parameters for tile drawing. RenderColor and BlendMode are not used.
 		 */
-		SLATEIM_API void DrawMaterial(UMaterialInterface* RenderMaterial, const FVector2D& CanvasPosition, const FVector2D& CanvasSize, const FTileRenderParams& RenderParams = FTileRenderParams());
+		void DrawMaterial(UMaterialInterface* RenderMaterial, const FVector2D& CanvasPosition, const FVector2D& CanvasSize, const FTileRenderParams& RenderParams = FTileRenderParams());
+
+		/**
+		 * Extra optional parameters for rendering text.
+		 */
+		struct FTextRenderParams
+		{
+			// The scale of the text on the canvas.
+			FVector2D Scale = FVector2D(1.0f, 1.0f);
+
+			// The color of the text.
+			FLinearColor RenderColor = FLinearColor::White;
+			
+			// Adjust the distance between letters.
+			float Kerning = 0.0f;
+
+			// The color of the text shadow.
+			FLinearColor ShadowColor = FLinearColor::Black;
+
+			// The offset of the shadow from the text it is shadowing.
+			FVector2D ShadowOffset = FVector2D::UnitVector;
+
+			// If true, the text is centered on the given draw position on the X axis.
+			bool bCentreX = false;
+
+			// If true, the text is centered on the given draw position on the Y axis.
+			bool bCentreY = false;
+
+			// If true, outline the text.
+			bool bOutlined = false;
+
+			// Color of the text outline.
+			FLinearColor OutlineColor = FLinearColor::Black;
+		};
 
 		/**
 		 * Queue drawing text to the canvas.
@@ -1491,9 +1146,9 @@ namespace SlateIM
 		 * @param RenderFont The font used to render the text.
 		 * @param RenderText The string to render.
 		 * @param CanvasPosition The position of the text on the canvas, in pixels.
-		 * @param RenderParams (optional) Additional parameters to draw the text.
+		 * @param RenderParams Extra optional parameters for rendering the text.
 		 */
-		SLATEIM_API void DrawText(UFont* RenderFont, const FString& RenderText, const FVector2D& CanvasPosition, const FTextRenderParams& RenderParams);
+		void DrawText(UFont* RenderFont, const FString& RenderText, const FVector2D& CanvasPosition, const FTextRenderParams& RenderParams);
 
 		/**
 		 * Queue drawing text to the canvas.
@@ -1506,7 +1161,7 @@ namespace SlateIM
 		 * @param Scale (optional) The scale of the text on the canvas.
 		 * @param RenderInfo (optional) Extra optional parameters for rendering the text.
 		 */
-		SLATEIM_API void DrawText(UFont* RenderFont, const FString& RenderText, const FVector2f& CanvasPosition, const FVector2f& Scale = FVector2f::UnitVector, const FFontRenderInfo& RenderInfo = FFontRenderInfo());
+		void DrawText(UFont* RenderFont, const FString& RenderText, const FVector2f& CanvasPosition, const FVector2f& Scale = FVector2f::UnitVector, const FFontRenderInfo& RenderInfo = FFontRenderInfo());
 
 		/**
 		 * Queue drawing wrapped text to canvas.
@@ -1524,7 +1179,24 @@ namespace SlateIM
 		 * @param Scale (optional) The scale of the text on the canvas.
 		 * @param RenderInfo (optional) Extra optional parameters for rendering the text.
 		 */
-		SLATEIM_API void DrawWrappedText(const UFont* RenderFont, const FString& RenderText, const FVector2f& CanvasPosition, bool bCenterTextX, bool bCenterTextY, const FVector2f& Scale = FVector2f::UnitVector, const FFontRenderInfo& RenderInfo = FFontRenderInfo());
+		void DrawWrappedText(const UFont* RenderFont, const FString& RenderText, const FVector2f& CanvasPosition, bool bCenterTextX, bool bCenterTextY, const FVector2f& Scale = FVector2f::UnitVector, const FFontRenderInfo& RenderInfo = FFontRenderInfo());
+
+		/**
+		 * Extra optional parameters for rendering borders.
+		 * 
+		 * Tile refers to the centre texture.
+		 */
+		struct FBorderRenderParams : FTileRenderParams
+		{
+			// The scale of the border.
+			FVector2D BorderScale = FVector2D(0.1, 0.1);
+
+			// The scale of the background.
+			FVector2D BackgroundScale = FVector2D(0.1, 0.1);
+
+			// Frame corner size, in percent (should be < 0.5f).
+			FVector2D CornerSize = FVector2D::ZeroVector;
+		};
 
 		/**
 		 * Queue drawing a border to the canvas.
@@ -1537,9 +1209,9 @@ namespace SlateIM
 		 * @param BottomBorderTexture The texture for the bottom of the border.
 		 * @param CanvasPosition The position of the border on the canvas, in pixels.
 		 * @param CanvasSize The size of the border on the canvas, in pixels.
-		 * @param RenderParams (optional) Additional parameters to draw the border.
+		 * @param RenderParams (optional) Extra optional params for rendering the border.
 		 */
-		SLATEIM_API void DrawBorder(UTexture* BorderTexture, UTexture* BackgroundTexture, UTexture* LeftBorderTexture, UTexture* RightBorderTexture, UTexture* TopBorderTexture, UTexture* BottomBorderTexture, const FVector2D& CanvasPosition, const FVector2D& CanvasSize, const FBorderRenderParams& RenderParams = FBorderRenderParams());
+		void DrawBorder(UTexture* BorderTexture, UTexture* BackgroundTexture, UTexture* LeftBorderTexture, UTexture* RightBorderTexture, UTexture* TopBorderTexture, UTexture* BottomBorderTexture, const FVector2D& CanvasPosition, const FVector2D& CanvasSize, const FBorderRenderParams& RenderParams = FBorderRenderParams());
 
 		/**
 		 * Queue drawing an unfilled box to the canvas.
@@ -1549,27 +1221,23 @@ namespace SlateIM
 		 * @param Thickness (optional) The thicknses of the lines to draw.
 		 * @param RenderColor (optional) The color of the lines to draw.
 		 */
-		SLATEIM_API void DrawBox(const FVector2D& CanvasPosition, const FVector2D& CanvasSize, float Thickness = 1.0f, FLinearColor RenderColor = FLinearColor::White);
+		void DrawBox(const FVector2D& CanvasPosition, const FVector2D& CanvasSize, float Thickness = 1.0f, FLinearColor RenderColor = FLinearColor::White);
 
 		/**
 		 * Queue drawing a series of texture-filled triangles to the canvas.
-		 *
+		 * 
 		 * @param RenderTexture The texture to draw with.
 		 * @param Triangles The list of triangles to draw.
-		 *
-		 * @note Not available in BPFL: TSharedRef<TArray<FCanvasUVTri>> is not Blueprint-compatible
 		 */
-		SLATEIM_API void DrawTriangles(UTexture* RenderTexture, TSharedRef<TArray<FCanvasUVTri>> Triangles);
+		void DrawTriangles(UTexture* RenderTexture, TSharedRef<TArray<FCanvasUVTri>> Triangles);
 
 		/**
 		 * Queue drawing a series of material-filled triangles to the canvas.
 		 *
 		 * @param RenderTexture The texture to draw with.
 		 * @param Triangles The list of triangles to draw.
-		 *
-		 * @note Not available in BPFL: TSharedRef<TArray<FCanvasUVTri>> is not Blueprint-compatible
 		 */
-		SLATEIM_API void DrawMaterialTriangles(UMaterialInterface* RenderMaterial, TSharedRef<TArray<FCanvasUVTri>> Triangles);
+		void DrawMaterialTriangles(UMaterialInterface* RenderMaterial, TSharedRef<TArray<FCanvasUVTri>> Triangles);
 
 		/**
 		 * Queue drawing a filled polygon to the canvas.
@@ -1580,7 +1248,7 @@ namespace SlateIM
 		 * @param NumberOfSides The number of sides of the polygon.
 		 * @param RenderColor (optional) The color of the polygon.
 		 */
-		SLATEIM_API void DrawPolygon(UTexture* RenderTexture, const FVector2D& CanvasPosition, const FVector2D& Radius, int32 NumberOfSides, FLinearColor RenderColor = FLinearColor::White);
+		void DrawPolygon(UTexture* RenderTexture, const FVector2D& CanvasPosition, const FVector2D& Radius, int32 NumberOfSides, FLinearColor RenderColor = FLinearColor::White);
 	}
 #endif // WITH_ENGINE
 #pragma endregion Canvas
diff --git a/Engine/Plugins/Experimental/SlateIM/Source/SlateIM/Public/SlateIMParameters.h b/Engine/Plugins/Experimental/SlateIM/Source/SlateIM/Public/SlateIMParameters.h
deleted file mode 100644
diff --git a/Engine/Plugins/Experimental/SlateIM/Source/SlateIM/Public/SlateIMParametersFwd.h b/Engine/Plugins/Experimental/SlateIM/Source/SlateIM/Public/SlateIMParametersFwd.h
deleted file mode 100644
diff --git a/Engine/Plugins/Experimental/SlateIM/Source/SlateIM/SlateIM.Build.cs b/Engine/Plugins/Experimental/SlateIM/Source/SlateIM/SlateIM.Build.cs
--- a/Engine/Plugins/Experimental/SlateIM/Source/SlateIM/SlateIM.Build.cs
+++ b/Engine/Plugins/Experimental/SlateIM/Source/SlateIM/SlateIM.Build.cs
@@ -27,8 +27,6 @@ public class SlateIM : ModuleRules
 
 		if (Target.bCompileAgainstEngine)
 		{
-			PublicDependencyModuleNames.Add("SlateIMEngine");
-
 			PrivateDependencyModuleNames.AddRange(
 				new string[]
 				{
@@ -37,7 +35,7 @@ public class SlateIM : ModuleRules
 				});
 		}
 
-		if (Target.bCompileAgainstEditor)
+		if(Target.bCompileAgainstEditor)
 		{
 			PrivateDependencyModuleNames.AddRange(
 				new string[]
diff --git a/Engine/Plugins/Experimental/SlateIM/Source/SlateIMBlueprint/Private/SlateIMBlueprint.cpp b/Engine/Plugins/Experimental/SlateIM/Source/SlateIMBlueprint/Private/SlateIMBlueprint.cpp
deleted file mode 100644
diff --git a/Engine/Plugins/Experimental/SlateIM/Source/SlateIMBlueprint/Public/SlateIMBlueprintFunctionLibrary.h b/Engine/Plugins/Experimental/SlateIM/Source/SlateIMBlueprint/Public/SlateIMBlueprintFunctionLibrary.h
deleted file mode 100644
diff --git a/Engine/Plugins/Experimental/SlateIM/Source/SlateIMBlueprint/SlateIMBlueprint.Build.cs b/Engine/Plugins/Experimental/SlateIM/Source/SlateIMBlueprint/SlateIMBlueprint.Build.cs
deleted file mode 100644
diff --git a/Engine/Plugins/Experimental/SlateIM/Source/SlateIMEngine/Private/Implementation/SlateIMEngineCanvasDrawCommandList.cpp b/Engine/Plugins/Experimental/SlateIM/Source/SlateIMEngine/Private/Implementation/SlateIMEngineCanvasDrawCommandList.cpp
deleted file mode 100644
diff --git a/Engine/Plugins/Experimental/SlateIM/Source/SlateIMEngine/Private/SlateIMEngineModule.cpp b/Engine/Plugins/Experimental/SlateIM/Source/SlateIMEngine/Private/SlateIMEngineModule.cpp
deleted file mode 100644
diff --git a/Engine/Plugins/Experimental/SlateIM/Source/SlateIMEngine/Public/Implementation/SlateIMEngineCanvasDrawCommandList.h b/Engine/Plugins/Experimental/SlateIM/Source/SlateIMEngine/Public/Implementation/SlateIMEngineCanvasDrawCommandList.h
deleted file mode 100644
diff --git a/Engine/Plugins/Experimental/SlateIM/Source/SlateIMEngine/Public/SlateIMEngineParameters.h b/Engine/Plugins/Experimental/SlateIM/Source/SlateIMEngine/Public/SlateIMEngineParameters.h
deleted file mode 100644
diff --git a/Engine/Plugins/Experimental/SlateIM/Source/SlateIMEngine/SlateIMEngine.Build.cs b/Engine/Plugins/Experimental/SlateIM/Source/SlateIMEngine/SlateIMEngine.Build.cs
deleted file mode 100644
