Remove parameters from constructor in SlateIM in game widgets and update all existing in widgets

#Synced-CL 52358810

diff --git a/Engine/Plugins/Experimental/SlateIM/Source/SlateIMInGame/Private/SlateIMInGameNomadTab.cpp b/Engine/Plugins/Experimental/SlateIM/Source/SlateIMInGame/Private/SlateIMInGameNomadTab.cpp
--- a/Engine/Plugins/Experimental/SlateIM/Source/SlateIMInGame/Private/SlateIMInGameNomadTab.cpp
+++ b/Engine/Plugins/Experimental/SlateIM/Source/SlateIMInGame/Private/SlateIMInGameNomadTab.cpp
@@ -8,17 +8,11 @@
 
 #include UE_INLINE_GENERATED_CPP_BY_NAME(SlateIMInGameNomadTab)
 
-ASlateIMInGameNomadTab::ASlateIMInGameNomadTab(const FObjectInitializer& ObjectInitializer)
+void ASlateIMInGameNomadTab::Init(const FName& InTabName, const FStringView& InTabTitle, const FSlateIcon& InTabIcon)
 {
-
-}
-
-ASlateIMInGameNomadTab::ASlateIMInGameNomadTab(const FName& InTabName, const FStringView& InTabTitle, const FSlateIcon& InTabIcon)
-	: TabName(InTabName)
-	, TabTitle(FText::FromStringView(InTabTitle))
-	, TabIcon(InTabIcon)
-{
-
+	TabName = InTabName;
+	TabTitle = FText::FromStringView(InTabTitle);
+	TabIcon = InTabIcon;
 }
 
 void ASlateIMInGameNomadTab::OnWidgetStarted()
diff --git a/Engine/Plugins/Experimental/SlateIM/Source/SlateIMInGame/Private/SlateIMInGameWindow.cpp b/Engine/Plugins/Experimental/SlateIM/Source/SlateIMInGame/Private/SlateIMInGameWindow.cpp
--- a/Engine/Plugins/Experimental/SlateIM/Source/SlateIMInGame/Private/SlateIMInGameWindow.cpp
+++ b/Engine/Plugins/Experimental/SlateIM/Source/SlateIMInGame/Private/SlateIMInGameWindow.cpp
@@ -6,14 +6,10 @@
 
 #include UE_INLINE_GENERATED_CPP_BY_NAME(SlateIMInGameWindow)
 
-ASlateIMInGameWindow::ASlateIMInGameWindow(const FObjectInitializer& ObjectInitializer)
-{
-}
-
-ASlateIMInGameWindow::ASlateIMInGameWindow(const FName& InWindowName, const FStringView& InWindowTitle)
-	: WindowName(InWindowName)
-	, WindowTitle(InWindowTitle)
+void ASlateIMInGameWindow::Init(const FName& InWindowName, const FStringView& InWindowTitle)
 {
+	WindowName = InWindowName;
+	WindowTitle = InWindowTitle;
 }
 
 void ASlateIMInGameWindow::DrawWidget(const float DeltaTime)
diff --git a/Engine/Plugins/Experimental/SlateIM/Source/SlateIMInGame/Public/SlateIMInGameNomadTab.h b/Engine/Plugins/Experimental/SlateIM/Source/SlateIMInGame/Public/SlateIMInGameNomadTab.h
--- a/Engine/Plugins/Experimental/SlateIM/Source/SlateIMInGame/Public/SlateIMInGameNomadTab.h
+++ b/Engine/Plugins/Experimental/SlateIM/Source/SlateIMInGame/Public/SlateIMInGameNomadTab.h
@@ -12,11 +12,9 @@ UCLASS(MinimalAPI)
 class ASlateIMInGameNomadTab : public ASlateIMInGameWidgetBase
 {
 	GENERATED_BODY()
-public:
-	ASlateIMInGameNomadTab(const FObjectInitializer& ObjectInitializer);
-	UE_API ASlateIMInGameNomadTab(const FName& InTabName, const FStringView& InTabTitle, const FSlateIcon& InTabIcon);
 
 protected:
+	UE_API void Init(const FName& InTabName, const FStringView& InTabTitle, const FSlateIcon& InTabIcon);
 	virtual void DrawContent(const float DeltaTime) {};
 	UE_API virtual void OnWidgetStarted() override;
 	UE_API virtual void OnWidgetStopped() override;
diff --git a/Engine/Plugins/Experimental/SlateIM/Source/SlateIMInGame/Public/SlateIMInGameWindow.h b/Engine/Plugins/Experimental/SlateIM/Source/SlateIMInGame/Public/SlateIMInGameWindow.h
--- a/Engine/Plugins/Experimental/SlateIM/Source/SlateIMInGame/Public/SlateIMInGameWindow.h
+++ b/Engine/Plugins/Experimental/SlateIM/Source/SlateIMInGame/Public/SlateIMInGameWindow.h
@@ -12,15 +12,15 @@ UCLASS(MinimalAPI)
 class ASlateIMInGameWindow : public ASlateIMInGameWidgetBase
 {
 	GENERATED_BODY()
-public:
-	ASlateIMInGameWindow(const FObjectInitializer& ObjectInitializer);
-	UE_API ASlateIMInGameWindow(const FName& InWindowName, const FStringView& InWindowTitle);
 
 protected:
-	UE_API virtual void DrawWidget(const float DeltaTime) override;
+	UE_API void Init(const FName& InWindowName, const FStringView& InWindowTitle);
 
 	virtual void DrawContent(const float DeltaTime) {};
 
+private:
+	UE_API virtual void DrawWidget(const float DeltaTime) override;
+
 	FName WindowName;
 	FString WindowTitle;
 	FVector2f WindowSize = FVector2f(500,500);
